Core: Moved Json out of Core and into its own module

The main changes are as follows:

1. moved Json out of Core into own module 'Json'
2. moved 3 i10n classes (Json serializers) from Core into a new module 'Internationalization' *
3. fixed up 2 i10n classes in Core to not instantiate the 3 Json-based classes. instead they are now passed in as a dependency

*) (2) and (3) were required to decouple the I10n code in Core from Json. Much of the i10n code probably doesn't belong into Core in the first place, but there is no time to fix this right now.

The following cosmetic changes were also made:

- NULL to nullptr
- namespaced enums to enum classes
- renamed the three i10n Json serializer classes to comply with naming conventions
- removed file header comments (not used)
- documentation, spelling, spacing etc,

#UpgradeNotes: If your module is including Json.h then you have to add 'Json' to your Build.cs module dependencies.

diff --git a/Engine/Plugins/2D/Paper2D/Source/Paper2DEditor/Paper2DEditor.Build.cs b/Engine/Plugins/2D/Paper2D/Source/Paper2DEditor/Paper2DEditor.Build.cs
--- a/Engine/Plugins/2D/Paper2D/Source/Paper2DEditor/Paper2DEditor.Build.cs
+++ b/Engine/Plugins/2D/Paper2D/Source/Paper2DEditor/Paper2DEditor.Build.cs
@@ -16,6 +16,7 @@ public class Paper2DEditor : ModuleRules
 			new string[] {
 				"Core",
 				"CoreUObject",
+				"Json",
 				"Slate",
 				"SlateCore",
 				"Engine",
diff --git a/Engine/Plugins/2D/Paper2D/Source/PaperJsonImporter/PaperJsonImporter.Build.cs b/Engine/Plugins/2D/Paper2D/Source/PaperJsonImporter/PaperJsonImporter.Build.cs
--- a/Engine/Plugins/2D/Paper2D/Source/PaperJsonImporter/PaperJsonImporter.Build.cs
+++ b/Engine/Plugins/2D/Paper2D/Source/PaperJsonImporter/PaperJsonImporter.Build.cs
@@ -10,6 +10,7 @@ public class PaperJsonImporter : ModuleRules
 			new string[] {
 				"Core",
 				"CoreUObject",
+				"Json",
 				"Slate",
 				"SlateCore",
 				"Engine",
diff --git a/Engine/Plugins/2D/Paper2D/Source/PaperTiledImporter/PaperTiledImporter.Build.cs b/Engine/Plugins/2D/Paper2D/Source/PaperTiledImporter/PaperTiledImporter.Build.cs
--- a/Engine/Plugins/2D/Paper2D/Source/PaperTiledImporter/PaperTiledImporter.Build.cs
+++ b/Engine/Plugins/2D/Paper2D/Source/PaperTiledImporter/PaperTiledImporter.Build.cs
@@ -11,6 +11,7 @@ public class PaperTiledImporter : ModuleRules
 				"Core",
 				"CoreUObject",
 				"Engine",
+				"Json",
 				"Paper2D",
 				"UnrealEd",
 				"Paper2DEditor"
diff --git a/Engine/Plugins/Editor/EpicSurvey/Source/EpicSurvey/EpicSurvey.Build.cs b/Engine/Plugins/Editor/EpicSurvey/Source/EpicSurvey/EpicSurvey.Build.cs
--- a/Engine/Plugins/Editor/EpicSurvey/Source/EpicSurvey/EpicSurvey.Build.cs
+++ b/Engine/Plugins/Editor/EpicSurvey/Source/EpicSurvey/EpicSurvey.Build.cs
@@ -26,6 +26,7 @@ namespace UnrealBuildTool.Rules
 				new string[]
 				{
 					"HTTP",
+					"Json",
 					"OnlineSubSystem",
 					"SlateCore",
 					"UnrealEd",
diff --git a/Engine/Plugins/Messaging/UdpMessaging/Source/UdpMessaging/UdpMessaging.Build.cs b/Engine/Plugins/Messaging/UdpMessaging/Source/UdpMessaging/UdpMessaging.Build.cs
--- a/Engine/Plugins/Messaging/UdpMessaging/Source/UdpMessaging/UdpMessaging.Build.cs
+++ b/Engine/Plugins/Messaging/UdpMessaging/Source/UdpMessaging/UdpMessaging.Build.cs
@@ -17,6 +17,7 @@ namespace UnrealBuildTool.Rules
 				new string[] {
 					"Core",
 					"CoreUObject",
+					"Json",
 					"Sockets",
 				}
 			);
diff --git a/Engine/Source/Developer/DesktopPlatform/DesktopPlatform.Build.cs b/Engine/Source/Developer/DesktopPlatform/DesktopPlatform.Build.cs
--- a/Engine/Source/Developer/DesktopPlatform/DesktopPlatform.Build.cs
+++ b/Engine/Source/Developer/DesktopPlatform/DesktopPlatform.Build.cs
@@ -8,7 +8,12 @@ public class DesktopPlatform : ModuleRules
 	{
 		PrivateIncludePaths.Add("Developer/DesktopPlatform/Private");
 
-		PrivateDependencyModuleNames.Add("Core");
+		PrivateDependencyModuleNames.AddRange(
+			new string[] {
+				"Core",
+				"Json",
+			}
+		);
 
 		if (Target.Platform == UnrealTargetPlatform.Linux)
 		{
diff --git a/Engine/Source/Developer/DesktopPlatform/Private/DesktopPlatformBase.h b/Engine/Source/Developer/DesktopPlatform/Private/DesktopPlatformBase.h
--- a/Engine/Source/Developer/DesktopPlatform/Private/DesktopPlatformBase.h
+++ b/Engine/Source/Developer/DesktopPlatform/Private/DesktopPlatformBase.h
@@ -3,7 +3,7 @@
 #pragma once
 
 #include "IDesktopPlatform.h"
-#include "Runtime/Core/Public/Serialization/Json/Json.h"
+#include "Json.h"
 #include "UProjectInfo.h"
 
 class FDesktopPlatformBase : public IDesktopPlatform
diff --git a/Engine/Source/Developer/LogVisualizer/LogVisualizer.Build.cs b/Engine/Source/Developer/LogVisualizer/LogVisualizer.Build.cs
--- a/Engine/Source/Developer/LogVisualizer/LogVisualizer.Build.cs
+++ b/Engine/Source/Developer/LogVisualizer/LogVisualizer.Build.cs
@@ -25,6 +25,7 @@ public class LogVisualizer : ModuleRules
 				"Core",
 				"CoreUObject",
                 "InputCore",
+				"Json",
 				"Slate",
 				"SlateCore",
                 "EditorStyle",
diff --git a/Engine/Source/Developer/SessionFrontend/SessionFrontend.Build.cs b/Engine/Source/Developer/SessionFrontend/SessionFrontend.Build.cs
--- a/Engine/Source/Developer/SessionFrontend/SessionFrontend.Build.cs
+++ b/Engine/Source/Developer/SessionFrontend/SessionFrontend.Build.cs
@@ -24,6 +24,7 @@ public class SessionFrontend : ModuleRules
 			new string[] {
 				"DesktopPlatform",
                 "InputCore",
+				"Json",
 				"SlateCore",
 
 				// @todo gmp: remove these dependencies by making the session front-end extensible
diff --git a/Engine/Source/Developer/SuperSearch/SuperSearch.Build.cs b/Engine/Source/Developer/SuperSearch/SuperSearch.Build.cs
--- a/Engine/Source/Developer/SuperSearch/SuperSearch.Build.cs
+++ b/Engine/Source/Developer/SuperSearch/SuperSearch.Build.cs
@@ -12,6 +12,7 @@ public class SuperSearch : ModuleRules
 				"CoreUObject", // @todo Mac: for some reason it's needed to link in debug on Mac
 				"Engine",
                 "InputCore",
+				"Json",
 				"UnrealEd",
 				"Slate",
 				"SlateCore",
diff --git a/Engine/Source/Editor/Kismet/Kismet.Build.cs b/Engine/Source/Editor/Kismet/Kismet.Build.cs
--- a/Engine/Source/Editor/Kismet/Kismet.Build.cs
+++ b/Engine/Source/Editor/Kismet/Kismet.Build.cs
@@ -27,6 +27,7 @@ public class Kismet : ModuleRules
 				"SlateCore",
                 "EditorStyle",
 				"Engine",
+				"Json",
 				"Merge",
 				"MessageLog",
 				"UnrealEd",
diff --git a/Engine/Source/Editor/Kismet/Private/FindInBlueprints.cpp b/Engine/Source/Editor/Kismet/Private/FindInBlueprints.cpp
--- a/Engine/Source/Editor/Kismet/Private/FindInBlueprints.cpp
+++ b/Engine/Source/Editor/Kismet/Private/FindInBlueprints.cpp
@@ -438,7 +438,7 @@ bool FFindInBlueprintsGraphNode::ExtractContent(TSharedPtr< FJsonObject > InJson
 {
 	// Very important to get the schema first, other bits of data depend on it
 	FString SchemaName;
-	TSharedPtr< FJsonValue > SchemaNameValue = InJsonNode->GetField< EJson::Type::String >(FFindInBlueprintSearchManager::ConvertFTextToHexString(FFindInBlueprintSearchTags::FiB_SchemaName));
+	TSharedPtr< FJsonValue > SchemaNameValue = InJsonNode->GetField< EJson::String >(FFindInBlueprintSearchManager::ConvertFTextToHexString(FFindInBlueprintSearchTags::FiB_SchemaName));
 	if(SchemaNameValue.IsValid())
 	{
 		SchemaName = FFindInBlueprintSearchManager::ConvertHexStringToFText(SchemaNameValue->AsString()).ToString();
diff --git a/Engine/Source/Editor/NewsFeed/NewsFeed.Build.cs b/Engine/Source/Editor/NewsFeed/NewsFeed.Build.cs
--- a/Engine/Source/Editor/NewsFeed/NewsFeed.Build.cs
+++ b/Engine/Source/Editor/NewsFeed/NewsFeed.Build.cs
@@ -17,6 +17,7 @@ public class NewsFeed : ModuleRules
 				"HTTP",
 				"ImageWrapper",
                 "InputCore",
+				"Json",
 				"OnlineSubsystem",
 				"Slate",
 				"SlateCore",
diff --git a/Engine/Source/Editor/TranslationEditor/Private/TranslationDataManager.cpp b/Engine/Source/Editor/TranslationEditor/Private/TranslationDataManager.cpp
--- a/Engine/Source/Editor/TranslationEditor/Private/TranslationDataManager.cpp
+++ b/Engine/Source/Editor/TranslationEditor/Private/TranslationDataManager.cpp
@@ -4,16 +4,18 @@
 #include "TranslationEditor.h"
 #include "WorkspaceMenuStructureModule.h"
 #include "TranslationUnit.h"
-#include "InternationalizationManifestJsonSerializer.h"
-#include "InternationalizationArchiveJsonSerializer.h"
 #include "ISourceControlModule.h"
 #include "MessageLog.h"
 #include "TextLocalizationManager.h"
+#include "JsonInternationalizationArchiveSerializer.h"
+#include "JsonInternationalizationManifestSerializer.h"
+
 
 DEFINE_LOG_CATEGORY_STATIC(LogTranslationEditor, Log, All);
 
 #define LOCTEXT_NAMESPACE "TranslationDataManager"
 
+
 FTranslationDataManager::FTranslationDataManager( const FString& InManifestFilePath, const FString& InArchiveFilePath )
 : ManifestFilePath(InManifestFilePath)
 , ArchiveFilePath(InArchiveFilePath)
@@ -50,7 +52,7 @@ FTranslationDataManager::FTranslationDataManager( const FString& InManifestFileP
 				GWarn->StatusUpdate(NumManifestEntriesParsed, ManifestEntriesCount, FText::Format(LOCTEXT("LoadingCurrentManifestEntries", "Loading Entry {0} of {1} from Current Translation Manifest..."), FText::AsNumber(NumManifestEntriesParsed), FText::AsNumber(ManifestEntriesCount)));
 				const TSharedRef<FManifestEntry> ManifestEntry = ManifestItr.Value();
 				UTranslationUnit* TranslationUnit = NewObject<UTranslationUnit>();
-				check(TranslationUnit != NULL);
+				check(TranslationUnit != nullptr);
 				// We want Undo/Redo support
 				TranslationUnit->SetFlags(RF_Transactional);
 				TranslationUnit->HasBeenReviewed = false;
@@ -131,7 +133,7 @@ TSharedPtr< FInternationalizationArchive > FTranslationDataManager::ReadArchive(
 	if ( !ArchiveJsonObject.IsValid() )
 	{
 		UE_LOG(LogTranslationEditor, Error, TEXT("Could not read archive file %s."), *ArchiveFilePath);
-		return NULL;
+		return nullptr;
 	}
 
 	TSharedRef< FInternationalizationArchive > InternationalizationArchive = MakeShareable( new FInternationalizationArchive );
@@ -148,7 +150,7 @@ TSharedPtr<FJsonObject> FTranslationDataManager::ReadJSONTextFile(const FString&
 	if ( !FFileHelper::LoadFileToString(FileContents, *InFilePath) )
 	{
 		UE_LOG(LogTranslationEditor, Error,TEXT("Failed to load file %s."), *InFilePath);
-		return NULL;
+		return nullptr;
 	}
 
 	//parse as JSON
@@ -159,7 +161,7 @@ TSharedPtr<FJsonObject> FTranslationDataManager::ReadJSONTextFile(const FString&
 	if( !FJsonSerializer::Deserialize( Reader, JSONObject ) || !JSONObject.IsValid())
 	{
 		UE_LOG(LogTranslationEditor, Error,TEXT("Invalid JSON in file %s."), *InFilePath);
-		return NULL;
+		return nullptr;
 	}
 
 	return JSONObject;
@@ -174,14 +176,14 @@ bool FTranslationDataManager::WriteTranslationData(bool bForceWrite /*= false*/)
 
 	for (UTranslationUnit* TranslationUnit : Untranslated)
 	{
-		if (TranslationUnit != NULL)
+		if (TranslationUnit != nullptr)
 		{
 			const FLocItem SearchSource(TranslationUnit->Source);
-			FString OldTranslation = Archive->FindEntryBySource(TranslationUnit->Namespace, SearchSource, NULL)->Translation.Text;
+			FString OldTranslation = Archive->FindEntryBySource(TranslationUnit->Namespace, SearchSource, nullptr)->Translation.Text;
 			FString TranslationToWrite = TranslationUnit->Translation;
 			if (!TranslationToWrite.Equals(OldTranslation))
 			{
-				Archive->SetTranslation(TranslationUnit->Namespace, TranslationUnit->Source, TranslationToWrite, NULL);
+				Archive->SetTranslation(TranslationUnit->Namespace, TranslationUnit->Source, TranslationToWrite, nullptr);
 				bNeedsWrite = true;
 			}
 		}
@@ -189,14 +191,14 @@ bool FTranslationDataManager::WriteTranslationData(bool bForceWrite /*= false*/)
 
 	for (UTranslationUnit* TranslationUnit : Review)
 	{
-		if (TranslationUnit != NULL)
+		if (TranslationUnit != nullptr)
 		{
 			const FLocItem SearchSource(TranslationUnit->Source);
-			FString OldTranslation = Archive->FindEntryBySource(TranslationUnit->Namespace, SearchSource, NULL)->Translation.Text;
+			FString OldTranslation = Archive->FindEntryBySource(TranslationUnit->Namespace, SearchSource, nullptr)->Translation.Text;
 			FString TranslationToWrite = TranslationUnit->Translation;
 			if (TranslationUnit->HasBeenReviewed && !TranslationToWrite.Equals(OldTranslation))
 			{
-				Archive->SetTranslation(TranslationUnit->Namespace, TranslationUnit->Source, TranslationToWrite, NULL);
+				Archive->SetTranslation(TranslationUnit->Namespace, TranslationUnit->Source, TranslationToWrite, nullptr);
 				bNeedsWrite = true;
 			}
 		}
@@ -204,14 +206,14 @@ bool FTranslationDataManager::WriteTranslationData(bool bForceWrite /*= false*/)
 
 	for (UTranslationUnit* TranslationUnit : Complete)
 	{
-		if (TranslationUnit != NULL)
+		if (TranslationUnit != nullptr)
 		{
 			const FLocItem SearchSource(TranslationUnit->Source);
-			FString OldTranslation = Archive->FindEntryBySource(TranslationUnit->Namespace, SearchSource, NULL)->Translation.Text;
+			FString OldTranslation = Archive->FindEntryBySource(TranslationUnit->Namespace, SearchSource, nullptr)->Translation.Text;
 			FString TranslationToWrite = TranslationUnit->Translation;
 			if (!TranslationToWrite.Equals(OldTranslation))
 			{
-				Archive->SetTranslation(TranslationUnit->Namespace, TranslationUnit->Source, TranslationToWrite, NULL);
+				Archive->SetTranslation(TranslationUnit->Namespace, TranslationUnit->Source, TranslationToWrite, nullptr);
 				bNeedsWrite = true;
 			}
 		}
@@ -410,7 +412,7 @@ void FTranslationDataManager::GetHistoryForTranslationUnits( TArray<UTranslation
 
 					for (UTranslationUnit* TranslationUnit : TranslationUnits)
 					{
-						if(TranslationUnit != NULL && TranslationUnit->Contexts.Num() > 0)
+						if(TranslationUnit != nullptr && TranslationUnit->Contexts.Num() > 0)
 						{
 							for (FTranslationContextInfo& ContextInfo : TranslationUnit->Contexts)
 							{
@@ -434,7 +436,7 @@ void FTranslationDataManager::GetHistoryForTranslationUnits( TArray<UTranslation
 								// Always add first instance of this string, and then add any versions that changed since
 								if (ContextInfo.Changes.Num() == 0 || !OldManifestEntryPtr->Source.Text.Equals(PreviousSourceText))
 								{
-									TSharedPtr< FArchiveEntry > OldArchiveEntry = ArchivePtr->FindEntryBySource(OldManifestEntryPtr->Namespace, OldManifestEntryPtr->Source, NULL);
+									TSharedPtr< FArchiveEntry > OldArchiveEntry = ArchivePtr->FindEntryBySource(OldManifestEntryPtr->Namespace, OldManifestEntryPtr->Source, nullptr);
 									if (OldArchiveEntry.IsValid())
 									{
 										FTranslationChange Change;
@@ -502,7 +504,10 @@ void FTranslationDataManager::PreviewAllTranslationsInEditor()
 
 	FString ConfigFilePath = ConfigDirectory / "Localization" / "Regenerate" + FPaths::GetBaseFilename(ManifestFilePath) + ".ini";
 
-	FTextLocalizationManager::Get().RegenerateResources(ConfigFilePath);
+	FJsonInternationalizationArchiveSerializer ArchiveSerializer;
+	FJsonInternationalizationManifestSerializer ManifestSerializer;
+
+	FTextLocalizationManager::Get().RegenerateResources(ConfigFilePath, ArchiveSerializer, ManifestSerializer);
 }
 
 void FTranslationDataManager::PopulateSearchResultsUsingFilter(const FString& SearchFilter)
@@ -511,7 +516,7 @@ void FTranslationDataManager::PopulateSearchResultsUsingFilter(const FString& Se
 
 	for (UTranslationUnit* TranslationUnit : AllTranslations)
 	{
-		if (TranslationUnit != NULL)
+		if (TranslationUnit != nullptr)
 		{
 			bool bAdded = false;
 			if (TranslationUnit->Source.Contains(SearchFilter) ||
@@ -562,7 +567,7 @@ void FTranslationDataManager::LoadFromArchive(TArray<UTranslationUnit*>& InTrans
 		for (int32 CurrentTranslationUnitIndex = 0; CurrentTranslationUnitIndex < TranslationUnits.Num(); ++CurrentTranslationUnitIndex)
 		{
 			UTranslationUnit* TranslationUnit = TranslationUnits[CurrentTranslationUnitIndex];
-			if (TranslationUnit != NULL)
+			if (TranslationUnit != nullptr)
 			{
 				if (!TranslationUnit->IsRooted())
 				{
@@ -573,7 +578,7 @@ void FTranslationDataManager::LoadFromArchive(TArray<UTranslationUnit*>& InTrans
 				GWarn->StatusUpdate(CurrentTranslationUnitIndex, TranslationUnits.Num(), FText::Format(LOCTEXT("LoadingCurrentArchiveEntries", "Loading Entry {0} of {1} from Translation Archive..."), FText::AsNumber(CurrentTranslationUnitIndex), FText::AsNumber(TranslationUnits.Num())));
 
 				const FLocItem SourceSearch(TranslationUnit->Source);
-				TSharedPtr<FArchiveEntry> ArchiveEntry = Archive->FindEntryBySource(TranslationUnit->Namespace, SourceSearch, NULL);
+				TSharedPtr<FArchiveEntry> ArchiveEntry = Archive->FindEntryBySource(TranslationUnit->Namespace, SourceSearch, nullptr);
 				if (ArchiveEntry.IsValid())
 				{
 					const FString PreviousTranslation = TranslationUnit->Translation;
diff --git a/Engine/Source/Editor/TranslationEditor/Private/TranslationDataManager.h b/Engine/Source/Editor/TranslationEditor/Private/TranslationDataManager.h
--- a/Engine/Source/Editor/TranslationEditor/Private/TranslationDataManager.h
+++ b/Engine/Source/Editor/TranslationEditor/Private/TranslationDataManager.h
@@ -7,8 +7,8 @@ class FInternationalizationArchive;
 class FInternationalizationManifest;
 class UTranslationUnit;
 
-#include "InternationalizationArchiveJsonSerializer.h"
-#include "InternationalizationManifestJsonSerializer.h"
+#include "JsonInternationalizationArchiveSerializer.h"
+#include "JsonInternationalizationManifestSerializer.h"
 
 class FTranslationDataManager : public TSharedFromThis<FTranslationDataManager>
 {
@@ -94,9 +94,9 @@ private:
 	TArray<UTranslationUnit*> ChangedOnImport;
 
 	/** Serializes and deserializes our Archive */
-	FInternationalizationArchiveJsonSerializer ArchiveSerializer;
+	FJsonInternationalizationArchiveSerializer ArchiveSerializer;
 	/** Serializes and deserializes Manifests */
-	FInternationalizationManifestJsonSerializer ManifestSerializer;
+	FJsonInternationalizationManifestSerializer ManifestSerializer;
 
 	/** Archive for the current project and translation language */
 	TSharedPtr< FInternationalizationArchive > ArchivePtr;
@@ -119,5 +119,3 @@ private:
 	/** Files that are already checked out from Perforce **/
 	TArray<FString> CheckedOutFiles;
 };
-
-
diff --git a/Engine/Source/Editor/TranslationEditor/TranslationEditor.Build.cs b/Engine/Source/Editor/TranslationEditor/TranslationEditor.Build.cs
--- a/Engine/Source/Editor/TranslationEditor/TranslationEditor.Build.cs
+++ b/Engine/Source/Editor/TranslationEditor/TranslationEditor.Build.cs
@@ -22,6 +22,8 @@ public class TranslationEditor : ModuleRules
 				"CoreUObject",
 				"Engine",
                 "InputCore",
+				"Internationalization",
+				"Json",
                 "PropertyEditor",
 				"Slate",
 				"SlateCore",
diff --git a/Engine/Source/Editor/UnrealEd/Private/Commandlets/GatherTextCommandletBase.cpp b/Engine/Source/Editor/UnrealEd/Private/Commandlets/GatherTextCommandletBase.cpp
--- a/Engine/Source/Editor/UnrealEd/Private/Commandlets/GatherTextCommandletBase.cpp
+++ b/Engine/Source/Editor/UnrealEd/Private/Commandlets/GatherTextCommandletBase.cpp
@@ -1,10 +1,10 @@
 // Copyright 1998-2014 Epic Games, Inc. All Rights Reserved.
 
 #include "UnrealEd.h"
-#include "Json.h"
-#include "InternationalizationManifestJsonSerializer.h"
 #include "Internationalization/InternationalizationMetadata.h"
-#include "InternationalizationMetadataJsonSerializer.h"
+#include "Json.h"
+#include "JsonInternationalizationManifestSerializer.h"
+#include "JsonInternationalizationMetadataSerializer.h"
 #include "ISourceControlModule.h"
 
 DEFINE_LOG_CATEGORY_STATIC(LogGatherTextCommandletBase, Log, All);
@@ -201,7 +201,7 @@ bool FManifestInfo::AddManifestDependency( const FString& InManifestFilePath )
 
 	if( FPaths::FileExists( InManifestFilePath ) )
 	{
-		FInternationalizationManifestJsonSerializer ManifestSerializer;
+		FJsonInternationalizationManifestSerializer ManifestSerializer;
 		TSharedPtr< FJsonObject > LoadedManifestDependencyJsonObject = UGatherTextCommandletBase::ReadJSONTextFile( InManifestFilePath );
 
 		TSharedRef< FInternationalizationManifest > LoadedManifestDependency = MakeShareable( new FInternationalizationManifest );
@@ -296,11 +296,11 @@ void FManifestInfo::ApplyManifestDependencies()
 						FString Message = UGatherTextCommandletBase::MungeLogOutput( FString::Printf(TEXT("Found previously entered localized string [%s] %s %s=\"%s\" %s. It was previously \"%s\" %s in dependency manifest %s."),
 							*ManifestEntry->Namespace,
 							*ContextIt->Key,
-							*FInternationalizationMetaDataJsonSerializer::MetadataToString( ContextIt->KeyMetadataObj ),
+							*FJsonInternationalizationMetaDataSerializer::MetadataToString( ContextIt->KeyMetadataObj ),
 							*ManifestEntry->Source.Text,
-							*FInternationalizationMetaDataJsonSerializer::MetadataToString(ManifestEntry->Source.MetadataObj),
+							*FJsonInternationalizationMetaDataSerializer::MetadataToString(ManifestEntry->Source.MetadataObj),
 							*DependencyEntry->Source.Text,
-							*FInternationalizationMetaDataJsonSerializer::MetadataToString(DependencyEntry->Source.MetadataObj),
+							*FJsonInternationalizationMetaDataSerializer::MetadataToString(DependencyEntry->Source.MetadataObj),
 							*DependencyFileName));
 						UE_LOG(LogGatherTextCommandletBase, Warning, TEXT("%s"), *Message);
 
@@ -323,7 +323,7 @@ void FManifestInfo::ApplyManifestDependencies()
 							*ManifestEntry->Namespace,
 							*ContextIt->Key,
 							*ManifestEntry->Source.Text,
-							*FInternationalizationMetaDataJsonSerializer::MetadataToString( ManifestEntry->Source.MetadataObj ));
+							*FJsonInternationalizationMetaDataSerializer::MetadataToString( ManifestEntry->Source.MetadataObj ));
 					}
 				}
 			}
@@ -360,11 +360,11 @@ bool FManifestInfo::AddEntry( const FString& EntryDescription, const FString& Na
 				*EntryDescription,
 				*Namespace,
 				*Context.Key,
-				*FInternationalizationMetaDataJsonSerializer::MetadataToString( Context.KeyMetadataObj ),
+				*FJsonInternationalizationMetaDataSerializer::MetadataToString( Context.KeyMetadataObj ),
 				*Source.Text,
-				*FInternationalizationMetaDataJsonSerializer::MetadataToString( Source.MetadataObj ),
+				*FJsonInternationalizationMetaDataSerializer::MetadataToString( Source.MetadataObj ),
 				*ExistingEntry->Source.Text,
-				*FInternationalizationMetaDataJsonSerializer::MetadataToString( ExistingEntry->Source.MetadataObj ),
+				*FJsonInternationalizationMetaDataSerializer::MetadataToString( ExistingEntry->Source.MetadataObj ),
 				*ExistingEntrySourceLocation));
 
 			UE_LOG(LogGatherTextCommandletBase, Warning, TEXT("%s"), *Message);
@@ -383,7 +383,7 @@ bool FManifestInfo::AddEntry( const FString& EntryDescription, const FString& Na
 				*Namespace,
 				*Context.Key,
 				*Source.Text,
-				*FInternationalizationMetaDataJsonSerializer::MetadataToString( Source.MetadataObj ));
+				*FJsonInternationalizationMetaDataSerializer::MetadataToString( Source.MetadataObj ));
 		}
 	}
 	return bAddSuccessful;
@@ -652,7 +652,7 @@ FString FConflictReportInfo::ToString()
 		
 		if(bAddToReport)
 		{
-			FString KeyMetadataString = FInternationalizationMetaDataJsonSerializer::MetadataToString( Conflict->KeyMetadataObj );
+			FString KeyMetadataString = FJsonInternationalizationMetaDataSerializer::MetadataToString( Conflict->KeyMetadataObj );
 			Report += FString::Printf(TEXT("%s - %s %s\n"), *Namespace, *Key, *KeyMetadataString );
 						
 			for(auto EntryIter = Conflict->EntriesBySourceLocation.CreateConstIterator(); EntryIter; ++EntryIter)
@@ -664,7 +664,7 @@ FString FConflictReportInfo::ToString()
 
 				const FString& SourceText = EntryIter.Value().Text.ReplaceCharWithEscapedChar();
 
-				FString SourceMetadataString =  FInternationalizationMetaDataJsonSerializer::MetadataToString( EntryIter.Value().MetadataObj );
+				FString SourceMetadataString =  FJsonInternationalizationMetaDataSerializer::MetadataToString( EntryIter.Value().MetadataObj );
 				Report += FString::Printf(TEXT("\t%s - \"%s\" %s\n"), *ProcessedSourceLocation, *SourceText, *SourceMetadataString);
 			}
 			Report += TEXT("\n");
diff --git a/Engine/Source/Editor/UnrealEd/Private/Commandlets/GenerateGatherArchiveCommandlet.cpp b/Engine/Source/Editor/UnrealEd/Private/Commandlets/GenerateGatherArchiveCommandlet.cpp
--- a/Engine/Source/Editor/UnrealEd/Private/Commandlets/GenerateGatherArchiveCommandlet.cpp
+++ b/Engine/Source/Editor/UnrealEd/Private/Commandlets/GenerateGatherArchiveCommandlet.cpp
@@ -6,8 +6,8 @@
 #include "Internationalization/InternationalizationArchive.h"
 #include "Internationalization/InternationalizationManifest.h"
 #include "Internationalization/InternationalizationMetadata.h"
-#include "InternationalizationArchiveJsonSerializer.h"
-#include "InternationalizationManifestJsonSerializer.h"
+#include "JsonInternationalizationArchiveSerializer.h"
+#include "JsonInternationalizationManifestSerializer.h"
 
 
 DEFINE_LOG_CATEGORY_STATIC(LogGenerateArchiveCommandlet, Log, All);
@@ -138,7 +138,7 @@ int32 UGenerateGatherArchiveCommandlet::Main( const FString& Params )
 		return -1;
 	}
 
-	FInternationalizationManifestJsonSerializer ManifestSerializer;
+	FJsonInternationalizationManifestSerializer ManifestSerializer;
 	TSharedRef< FInternationalizationManifest > InternationalizationManifest = MakeShareable( new FInternationalizationManifest );
 
 	ManifestSerializer.DeserializeManifest( ManifestJsonObject.ToSharedRef(), InternationalizationManifest );
@@ -149,7 +149,7 @@ int32 UGenerateGatherArchiveCommandlet::Main( const FString& Params )
 		BuildArchiveFromManifest( InternationalizationManifest, InternationalizationArchive, SourceCulture, CulturesToGenerate[Culture] );
 
 		const FString CulturePath = DestinationPath / CulturesToGenerate[Culture];
-		FInternationalizationArchiveJsonSerializer ArchiveSerializer;
+		FJsonInternationalizationArchiveSerializer ArchiveSerializer;
 		TSharedRef< FInternationalizationArchive > OutputInternationalizationArchive = MakeShareable( new FInternationalizationArchive );
 
 		// Read in any existing archive for this culture.
@@ -161,13 +161,13 @@ int32 UGenerateGatherArchiveCommandlet::Main( const FString& Params )
 			ExistingArchiveJsonObject = ReadJSONTextFile( ExistingArchiveFileName );
 			
 			// Some of the existing archives were saved out with an "Unnamed" namespace for the root instead of the empty string.  We try to fix that here.
-			if( ExistingArchiveJsonObject->HasField( FInternationalizationArchiveJsonSerializer::TAG_NAMESPACE ) )
+			if( ExistingArchiveJsonObject->HasField( FJsonInternationalizationArchiveSerializer::TAG_NAMESPACE ) )
 			{
-				FString RootNamespace = ExistingArchiveJsonObject->GetStringField( FInternationalizationArchiveJsonSerializer::TAG_NAMESPACE );
+				FString RootNamespace = ExistingArchiveJsonObject->GetStringField( FJsonInternationalizationArchiveSerializer::TAG_NAMESPACE );
 				if( RootNamespace == TEXT("Unnamed") )
 				{
-					ExistingArchiveJsonObject->RemoveField( FInternationalizationArchiveJsonSerializer::TAG_NAMESPACE );
-					ExistingArchiveJsonObject->SetStringField( FInternationalizationArchiveJsonSerializer::TAG_NAMESPACE, TEXT("") );
+					ExistingArchiveJsonObject->RemoveField( FJsonInternationalizationArchiveSerializer::TAG_NAMESPACE );
+					ExistingArchiveJsonObject->SetStringField( FJsonInternationalizationArchiveSerializer::TAG_NAMESPACE, TEXT("") );
 				}
 			}
 
@@ -178,15 +178,15 @@ int32 UGenerateGatherArchiveCommandlet::Main( const FString& Params )
 				static bool PurgeNamespaceOfEmptyEntries(const TSharedPtr<FJsonObject>& JSONObject)
 				{
 					bool ModifiedChildrenArray = false;
-					if( JSONObject->HasField( FInternationalizationArchiveJsonSerializer::TAG_CHILDREN ) )
+					if( JSONObject->HasField( FJsonInternationalizationArchiveSerializer::TAG_CHILDREN ) )
 					{
-						TArray<TSharedPtr<FJsonValue>> ChildrenArray = JSONObject->GetArrayField(FInternationalizationArchiveJsonSerializer::TAG_CHILDREN);
+						TArray<TSharedPtr<FJsonValue>> ChildrenArray = JSONObject->GetArrayField(FJsonInternationalizationArchiveSerializer::TAG_CHILDREN);
 						for( int32 ChildIndex = ChildrenArray.Num() - 1; ChildIndex >= 0; --ChildIndex )
 						{
 							TSharedPtr<FJsonObject> Child = ChildrenArray[ ChildIndex ]->AsObject();
-							TSharedPtr<FJsonObject> TranslationObject = Child->GetObjectField(FInternationalizationArchiveJsonSerializer::TAG_TRANSLATION);
+							TSharedPtr<FJsonObject> TranslationObject = Child->GetObjectField(FJsonInternationalizationArchiveSerializer::TAG_TRANSLATION);
 
-							const FString& TranslatedText = TranslationObject->GetStringField(FInternationalizationArchiveJsonSerializer::TAG_TRANSLATION_TEXT);
+							const FString& TranslatedText = TranslationObject->GetStringField(FJsonInternationalizationArchiveSerializer::TAG_TRANSLATION_TEXT);
 
 							if(TranslatedText.IsEmpty())
 							{
@@ -197,26 +197,26 @@ int32 UGenerateGatherArchiveCommandlet::Main( const FString& Params )
 						}
 						if(ModifiedChildrenArray)
 						{
-							JSONObject->RemoveField(FInternationalizationArchiveJsonSerializer::TAG_CHILDREN);
+							JSONObject->RemoveField(FJsonInternationalizationArchiveSerializer::TAG_CHILDREN);
 							if(ChildrenArray.Num())
 							{
-								JSONObject->SetArrayField(FInternationalizationArchiveJsonSerializer::TAG_CHILDREN, ChildrenArray);
+								JSONObject->SetArrayField(FJsonInternationalizationArchiveSerializer::TAG_CHILDREN, ChildrenArray);
 							}
 						}
 					}
 
 					bool ModifiedSubnamespaceArray = false;
-					if( JSONObject->HasField( FInternationalizationArchiveJsonSerializer::TAG_SUBNAMESPACES ) )
+					if( JSONObject->HasField( FJsonInternationalizationArchiveSerializer::TAG_SUBNAMESPACES ) )
 					{
-						TArray<TSharedPtr<FJsonValue>> SubnamespaceArray = JSONObject->GetArrayField(FInternationalizationArchiveJsonSerializer::TAG_SUBNAMESPACES);
+						TArray<TSharedPtr<FJsonValue>> SubnamespaceArray = JSONObject->GetArrayField(FJsonInternationalizationArchiveSerializer::TAG_SUBNAMESPACES);
 
 						for( int32 Index = SubnamespaceArray.Num() - 1; Index >= 0; --Index )
 						{
 							TSharedPtr<FJsonObject> Subnamespace = SubnamespaceArray[ Index ]->AsObject();
 							ModifiedSubnamespaceArray = PurgeNamespaceOfEmptyEntries(Subnamespace);
 
-							bool HasChildren = Subnamespace->HasField( FInternationalizationArchiveJsonSerializer::TAG_CHILDREN );
-							bool HasSubnamespaces = Subnamespace->HasField( FInternationalizationArchiveJsonSerializer::TAG_SUBNAMESPACES );
+							bool HasChildren = Subnamespace->HasField( FJsonInternationalizationArchiveSerializer::TAG_CHILDREN );
+							bool HasSubnamespaces = Subnamespace->HasField( FJsonInternationalizationArchiveSerializer::TAG_SUBNAMESPACES );
 
 							if(!HasChildren && !HasSubnamespaces)
 							{
@@ -227,10 +227,10 @@ int32 UGenerateGatherArchiveCommandlet::Main( const FString& Params )
 						}
 						if(ModifiedSubnamespaceArray)
 						{
-							JSONObject->RemoveField(FInternationalizationArchiveJsonSerializer::TAG_SUBNAMESPACES);
+							JSONObject->RemoveField(FJsonInternationalizationArchiveSerializer::TAG_SUBNAMESPACES);
 							if(SubnamespaceArray.Num())
 							{
-								JSONObject->SetArrayField(FInternationalizationArchiveJsonSerializer::TAG_SUBNAMESPACES, SubnamespaceArray);
+								JSONObject->SetArrayField(FJsonInternationalizationArchiveSerializer::TAG_SUBNAMESPACES, SubnamespaceArray);
 							}
 						}
 					}
diff --git a/Engine/Source/Editor/UnrealEd/Private/Commandlets/GenerateGatherManifestCommandlet.cpp b/Engine/Source/Editor/UnrealEd/Private/Commandlets/GenerateGatherManifestCommandlet.cpp
--- a/Engine/Source/Editor/UnrealEd/Private/Commandlets/GenerateGatherManifestCommandlet.cpp
+++ b/Engine/Source/Editor/UnrealEd/Private/Commandlets/GenerateGatherManifestCommandlet.cpp
@@ -2,9 +2,9 @@
 
 #include "UnrealEd.h"
 #include "ISourceControlModule.h"
-#include "Json.h"
 #include "Internationalization/InternationalizationManifest.h"
-#include "InternationalizationManifestJsonSerializer.h"
+#include "Json.h"
+#include "JsonInternationalizationManifestSerializer.h"
 
 DEFINE_LOG_CATEGORY_STATIC(LogGenerateManifestCommandlet, Log, All);
 
@@ -126,7 +126,7 @@ bool UGenerateGatherManifestCommandlet::WriteManifest( const TSharedPtr<FInterna
 	{
 		return false;
 	}
-	FInternationalizationManifestJsonSerializer ManifestSerializer;
+	FJsonInternationalizationManifestSerializer ManifestSerializer;
 	TSharedRef<FJsonObject> JsonManifestObj = MakeShareable( new FJsonObject );
 	
 	bool bSuccess = ManifestSerializer.SerializeManifest( InManifest.ToSharedRef(), JsonManifestObj );
diff --git a/Engine/Source/Editor/UnrealEd/Private/Commandlets/GenerateTextLocalizationReportCommandlet.cpp b/Engine/Source/Editor/UnrealEd/Private/Commandlets/GenerateTextLocalizationReportCommandlet.cpp
--- a/Engine/Source/Editor/UnrealEd/Private/Commandlets/GenerateTextLocalizationReportCommandlet.cpp
+++ b/Engine/Source/Editor/UnrealEd/Private/Commandlets/GenerateTextLocalizationReportCommandlet.cpp
@@ -3,9 +3,9 @@
 #include "UnrealEd.h"
 #include "Json.h"
 #include "InternationalizationManifest.h"
-#include "InternationalizationManifestJsonSerializer.h"
 #include "InternationalizationArchive.h"
-#include "InternationalizationArchiveJsonSerializer.h"
+#include "JsonInternationalizationManifestSerializer.h"
+#include "JsonInternationalizationArchiveSerializer.h"
 #include "BreakIterator.h"
 
 DEFINE_LOG_CATEGORY_STATIC(LogGenerateTextLocalizationReportCommandlet, Log, All);
@@ -223,7 +223,7 @@ bool UGenerateTextLocalizationReportCommandlet::ProcessWordCountReport(const FSt
 
 	TSharedRef<FInternationalizationManifest> InternationalizationManifest = MakeShareable( new FInternationalizationManifest );
 	{
-		FInternationalizationManifestJsonSerializer ManifestSerializer;
+		FJsonInternationalizationManifestSerializer ManifestSerializer;
 		ManifestSerializer.DeserializeManifest(ManifestJSONObject.ToSharedRef(), InternationalizationManifest);
 	}
 
@@ -321,7 +321,7 @@ bool UGenerateTextLocalizationReportCommandlet::ProcessWordCountReport(const FSt
 
 			TSharedRef<FInternationalizationArchive> InternationalizationArchive = MakeShareable( new FInternationalizationArchive );
 			{
-				FInternationalizationArchiveJsonSerializer ArchiveSerializer;
+				FJsonInternationalizationArchiveSerializer ArchiveSerializer;
 				ArchiveSerializer.DeserializeArchive(ArchiveJSONObject.ToSharedRef(), InternationalizationArchive);
 			}
 
diff --git a/Engine/Source/Editor/UnrealEd/Private/Commandlets/GenerateTextLocalizationResourceCommandlet.cpp b/Engine/Source/Editor/UnrealEd/Private/Commandlets/GenerateTextLocalizationResourceCommandlet.cpp
--- a/Engine/Source/Editor/UnrealEd/Private/Commandlets/GenerateTextLocalizationResourceCommandlet.cpp
+++ b/Engine/Source/Editor/UnrealEd/Private/Commandlets/GenerateTextLocalizationResourceCommandlet.cpp
@@ -4,9 +4,9 @@
 #include "TextLocalizationResourceGenerator.h"
 #include "Json.h"
 #include "InternationalizationManifest.h"
-#include "InternationalizationManifestJsonSerializer.h"
 #include "InternationalizationArchive.h"
-#include "InternationalizationArchiveJsonSerializer.h"
+#include "JsonInternationalizationManifestSerializer.h"
+#include "JsonInternationalizationArchiveSerializer.h"
 
 DEFINE_LOG_CATEGORY_STATIC(LogGenerateTextLocalizationResourceCommandlet, Log, All);
 
@@ -135,7 +135,7 @@ int32 UGenerateTextLocalizationResourceCommandlet::Main(const FString& Params)
 	}
 	TSharedRef<FInternationalizationManifest> InternationalizationManifest = MakeShareable( new FInternationalizationManifest );
 	{
-		FInternationalizationManifestJsonSerializer ManifestSerializer;
+		FJsonInternationalizationManifestSerializer ManifestSerializer;
 		ManifestSerializer.DeserializeManifest(ManifestJSONObject.ToSharedRef(), InternationalizationManifest);
 	}
 
@@ -161,7 +161,9 @@ int32 UGenerateTextLocalizationResourceCommandlet::Main(const FString& Params)
 		TAutoPtr<FArchive> TextLocalizationResourceArchive(IFileManager::Get().CreateFileWriter(*TextLocalizationResourcePath));
 		if (TextLocalizationResourceArchive)
 		{
-			if( !(FTextLocalizationResourceGenerator::Generate(SourcePath, InternationalizationManifest, CultureName, TextLocalizationResourceArchive)) )
+			FJsonInternationalizationArchiveSerializer ArchiveSerializer;
+
+			if( !(FTextLocalizationResourceGenerator::Generate(SourcePath, InternationalizationManifest, CultureName, TextLocalizationResourceArchive, ArchiveSerializer)) )
 			{
 				IFileManager::Get().Delete( *TextLocalizationResourcePath );
 			}
diff --git a/Engine/Source/Editor/UnrealEd/Private/Commandlets/InternationalizationConditioningCommandlet.cpp b/Engine/Source/Editor/UnrealEd/Private/Commandlets/InternationalizationConditioningCommandlet.cpp
--- a/Engine/Source/Editor/UnrealEd/Private/Commandlets/InternationalizationConditioningCommandlet.cpp
+++ b/Engine/Source/Editor/UnrealEd/Private/Commandlets/InternationalizationConditioningCommandlet.cpp
@@ -3,8 +3,8 @@
 #include "UnrealEd.h"
 #include "Internationalization/InternationalizationArchive.h"
 #include "Internationalization/InternationalizationManifest.h"
-#include "InternationalizationArchiveJsonSerializer.h"
-#include "InternationalizationManifestJsonSerializer.h"
+#include "JsonInternationalizationArchiveSerializer.h"
+#include "JsonInternationalizationManifestSerializer.h"
 #include "ConfigCacheIni.h"
 
 DEFINE_LOG_CATEGORY_STATIC(LogInternationalizationConditioningCommandlet, Log, All);
@@ -336,7 +336,7 @@ bool UInternationalizationConditioningCommandlet::ProcessManifest( const FString
 
 	// First we want to see if there is an existing manifest.  If so we will load it up and add our entries there
 	TSharedRef< FInternationalizationManifest > InternationalizationManifest = MakeShareable( new FInternationalizationManifest );
-	FInternationalizationManifestJsonSerializer ManifestSerializer;
+	FJsonInternationalizationManifestSerializer ManifestSerializer;
 
 	FString ExistingManifestFileName = DestinationPath / ManifestName;
 
@@ -448,7 +448,7 @@ bool UInternationalizationConditioningCommandlet::ProcessArchive( const FString&
 		}
 
 		TSharedRef< FInternationalizationArchive > InternationalizationArchive = MakeShareable( new FInternationalizationArchive );
-		FInternationalizationArchiveJsonSerializer ArchiveSerializer;
+		FJsonInternationalizationArchiveSerializer ArchiveSerializer;
 
 		const FString DestinationArchiveFileName = DestinationPath / TargetSubfolder / ArchiveName;
 
diff --git a/Engine/Source/Editor/UnrealEd/Private/Commandlets/InternationalizationExportCommandlet.cpp b/Engine/Source/Editor/UnrealEd/Private/Commandlets/InternationalizationExportCommandlet.cpp
--- a/Engine/Source/Editor/UnrealEd/Private/Commandlets/InternationalizationExportCommandlet.cpp
+++ b/Engine/Source/Editor/UnrealEd/Private/Commandlets/InternationalizationExportCommandlet.cpp
@@ -2,8 +2,8 @@
 
 #include "UnrealEd.h"
 #include "Culture.h"
-#include "InternationalizationManifestJsonSerializer.h"
-#include "InternationalizationArchiveJsonSerializer.h"
+#include "JsonInternationalizationManifestSerializer.h"
+#include "JsonInternationalizationArchiveSerializer.h"
 
 DEFINE_LOG_CATEGORY_STATIC(LogInternationalizationExportCommandlet, Log, All);
 
@@ -1079,7 +1079,7 @@ bool UInternationalizationExportCommandlet::DoExport( const FString& SourcePath,
 			return false;
 		}
 
-		FInternationalizationManifestJsonSerializer ManifestSerializer;
+		FJsonInternationalizationManifestSerializer ManifestSerializer;
 		ManifestSerializer.DeserializeManifest( ManifestJsonObject.ToSharedRef(), InternationalizationManifest );
 	}
 
@@ -1097,7 +1097,7 @@ bool UInternationalizationExportCommandlet::DoExport( const FString& SourcePath,
 		{
 			ArchiveJsonObject = ReadJSONTextFile( ArchiveFileName );
 
-			FInternationalizationArchiveJsonSerializer ArchiveSerializer;
+			FJsonInternationalizationArchiveSerializer ArchiveSerializer;
 			TSharedRef< FInternationalizationArchive > InternationalizationArchive = MakeShareable( new FInternationalizationArchive );
 			ArchiveSerializer.DeserializeArchive( ArchiveJsonObject.ToSharedRef(), InternationalizationArchive );
 
@@ -1261,7 +1261,7 @@ bool UInternationalizationExportCommandlet::DoImport(const FString& SourcePath,
 		TSharedPtr< FJsonObject > ArchiveJsonObject = NULL;
 		ArchiveJsonObject = ReadJSONTextFile( ArchiveFileName );
 
-		FInternationalizationArchiveJsonSerializer ArchiveSerializer;
+		FJsonInternationalizationArchiveSerializer ArchiveSerializer;
 		TSharedRef< FInternationalizationArchive > InternationalizationArchive = MakeShareable( new FInternationalizationArchive );
 		ArchiveSerializer.DeserializeArchive( ArchiveJsonObject.ToSharedRef(), InternationalizationArchive );
 
diff --git a/Engine/Source/Editor/UnrealEd/Private/Commandlets/RepairLocalizationDataCommandlet.cpp b/Engine/Source/Editor/UnrealEd/Private/Commandlets/RepairLocalizationDataCommandlet.cpp
--- a/Engine/Source/Editor/UnrealEd/Private/Commandlets/RepairLocalizationDataCommandlet.cpp
+++ b/Engine/Source/Editor/UnrealEd/Private/Commandlets/RepairLocalizationDataCommandlet.cpp
@@ -3,9 +3,9 @@
 #include "UnrealEd.h"
 #include "Commandlets/RepairLocalizationDataCommandlet.h"
 #include "InternationalizationManifest.h"
-#include "InternationalizationManifestJsonSerializer.h"
 #include "InternationalizationArchive.h"
-#include "InternationalizationArchiveJsonSerializer.h"
+#include "JsonInternationalizationManifestSerializer.h"
+#include "JsonInternationalizationArchiveSerializer.h"
 #include "Regex.h"
 
 DEFINE_LOG_CATEGORY_STATIC(LogRepairLocalizationDataCommandlet, Log, All);
@@ -241,7 +241,7 @@ int32 URepairLocalizationDataCommandlet::Main(const FString& Params)
 		UE_LOG(LogRepairLocalizationDataCommandlet, Error, TEXT("Could not read manifest file %s."), *ManifestFilePath);
 		return -1;
 	}
-	FInternationalizationManifestJsonSerializer ManifestSerializer;
+	FJsonInternationalizationManifestSerializer ManifestSerializer;
 	const TSharedRef< FInternationalizationManifest > InternationalizationManifest = MakeShareable( new FInternationalizationManifest );
 	ManifestSerializer.DeserializeManifest( ManifestJsonObject.ToSharedRef(), InternationalizationManifest );
 
@@ -261,7 +261,7 @@ int32 URepairLocalizationDataCommandlet::Main(const FString& Params)
 				UE_LOG(LogRepairLocalizationDataCommandlet, Error, TEXT("Could not read archive file %s."), *ArchiveFilePath);
 				return -1;
 			}
-			FInternationalizationArchiveJsonSerializer ArchiveSerializer;
+			FJsonInternationalizationArchiveSerializer ArchiveSerializer;
 			const TSharedRef< FInternationalizationArchive > InternationalizationArchive = MakeShareable( new FInternationalizationArchive );
 			ArchiveSerializer.DeserializeArchive( ArchiveJsonObject.ToSharedRef(), InternationalizationArchive );
 
@@ -288,7 +288,7 @@ int32 URepairLocalizationDataCommandlet::Main(const FString& Params)
 		for (const TSharedRef<FInternationalizationArchive>& InternationalizationArchive : InternationalizationArchives)
 		{
 			TSharedRef< FJsonObject > OutputArchiveJsonObj = MakeShareable( new FJsonObject );
-			FInternationalizationArchiveJsonSerializer ArchiveSerializer;
+			FJsonInternationalizationArchiveSerializer ArchiveSerializer;
 			ArchiveSerializer.SerializeArchive( InternationalizationArchive, OutputArchiveJsonObj );
 
 			int32 Culture = 0;
diff --git a/Engine/Source/Editor/UnrealEd/Private/Tests/EditorAutomationTests.cpp b/Engine/Source/Editor/UnrealEd/Private/Tests/EditorAutomationTests.cpp
--- a/Engine/Source/Editor/UnrealEd/Private/Tests/EditorAutomationTests.cpp
+++ b/Engine/Source/Editor/UnrealEd/Private/Tests/EditorAutomationTests.cpp
@@ -563,7 +563,7 @@ bool FJsonAutomationTest::RunTest(const FString& Parameters)
 		const TArray< TSharedPtr<FJsonValue> > Array = (*ArrayValue)->AsArray();
 		check(Array.Num() == 8);
 
-		EJson::Type ValueTypes[] = {EJson::Array, EJson::String, EJson::String, EJson::Null,
+		EJson ValueTypes[] = {EJson::Array, EJson::String, EJson::String, EJson::Null,
 			EJson::Boolean, EJson::Boolean, EJson::Number, EJson::Object};
 		for (int32 i = 0; i < Array.Num(); ++i)
 		{
@@ -638,7 +638,7 @@ bool FJsonAutomationTest::RunTest(const FString& Parameters)
 			TEXT("}");
 		TSharedRef< TJsonReader<> > Reader = TJsonReaderFactory<>::Create( InputString );
 
-		EJsonNotation::Type Notation = EJsonNotation::Null;
+		EJsonNotation Notation = EJsonNotation::Null;
 		check( Reader->ReadNext( Notation ) && Notation == EJsonNotation::ObjectStart );
 		check( Reader->GetLineNumber() == 1 && Reader->GetCharacterNumber() == 1 );
 
diff --git a/Engine/Source/Editor/UnrealEd/Private/UnrealEdSrv.cpp b/Engine/Source/Editor/UnrealEd/Private/UnrealEdSrv.cpp
--- a/Engine/Source/Editor/UnrealEd/Private/UnrealEdSrv.cpp
+++ b/Engine/Source/Editor/UnrealEd/Private/UnrealEdSrv.cpp
@@ -34,6 +34,9 @@
 #include "Landscape/LandscapeComponent.h"
 #include "Landscape/LandscapeHeightfieldCollisionComponent.h"
 #include "ComponentReregisterContext.h"
+#include "JsonInternationalizationArchiveSerializer.h"
+#include "JsonInternationalizationManifestSerializer.h"
+
 
 DEFINE_LOG_CATEGORY_STATIC(LogUnrealEdSrv, Log, All);
 
@@ -43,6 +46,7 @@ DEFINE_LOG_CATEGORY_STATIC(LogUnrealEdSrv, Log, All);
 static TCHAR TempStr[MAX_EDCMD], TempName[MAX_EDCMD], Temp[MAX_EDCMD];
 static uint16 Word1, Word4;
 
+
 /**
  * Dumps a set of selected objects to debugf.
  */
@@ -760,7 +764,10 @@ bool UUnrealEdEngine::Exec( UWorld* InWorld, const TCHAR* Stream, FOutputDevice&
 		FString ConfigFilePath;
 		if(FParse::Value(Str,TEXT("REGENLOC="), ConfigFilePath))
 		{
-			FTextLocalizationManager::Get().RegenerateResources(ConfigFilePath);
+			FJsonInternationalizationArchiveSerializer ArchiveSerializer;
+			FJsonInternationalizationManifestSerializer ManifestSerializer;
+
+			FTextLocalizationManager::Get().RegenerateResources(ConfigFilePath, ArchiveSerializer, ManifestSerializer);
 		}
 	}
 #endif
@@ -2976,7 +2983,10 @@ bool UUnrealEdEngine::Exec_Mode( const TCHAR* Str, FOutputDevice& Ar )
 	FString ConfigFilePath;
 	if( FParse::Value(Str,TEXT("REGENLOC="), ConfigFilePath) )
 	{
-		FTextLocalizationManager::Get().RegenerateResources(ConfigFilePath);
+		FJsonInternationalizationArchiveSerializer ArchiveSerializer;
+		FJsonInternationalizationManifestSerializer ManifestSerializer;
+
+		FTextLocalizationManager::Get().RegenerateResources(ConfigFilePath, ArchiveSerializer, ManifestSerializer);
 	}
 #endif
 
diff --git a/Engine/Source/Editor/UnrealEd/UnrealEd.Build.cs b/Engine/Source/Editor/UnrealEd/UnrealEd.Build.cs
--- a/Engine/Source/Editor/UnrealEd/UnrealEd.Build.cs
+++ b/Engine/Source/Editor/UnrealEd/UnrealEd.Build.cs
@@ -61,6 +61,7 @@ public class UnrealEd : ModuleRules
 				"CoreUObject",
 				"Documentation",
 				"Engine",
+				"Json",
 				"Projects",
 				"SandboxFile",
 				"Slate",
@@ -88,6 +89,7 @@ public class UnrealEd : ModuleRules
 				"EngineSettings",
 				"InputCore",
 				"InputBindingEditor",
+				"Internationalization",
 				"LauncherAutomatedService",
 				"LauncherServices",
 				"MaterialEditor",
diff --git a/Engine/Source/Programs/CrashReporter/CrashReportClient/CrashReportClient.Build.cs b/Engine/Source/Programs/CrashReporter/CrashReportClient/CrashReportClient.Build.cs
--- a/Engine/Source/Programs/CrashReporter/CrashReportClient/CrashReportClient.Build.cs
+++ b/Engine/Source/Programs/CrashReporter/CrashReportClient/CrashReportClient.Build.cs
@@ -21,6 +21,7 @@ public class CrashReportClient : ModuleRules
 				"Core",
 				"CrashDebugHelper",
 				"HTTP",
+				"Json",
 				"Projects",
 				"XmlParser",
 				"Analytics",
diff --git a/Engine/Source/Programs/Mac/UE4EditorServices/Private/Mac/UE4EditorServicesAppDelegate.cpp b/Engine/Source/Programs/Mac/UE4EditorServices/Private/Mac/UE4EditorServicesAppDelegate.cpp
--- a/Engine/Source/Programs/Mac/UE4EditorServices/Private/Mac/UE4EditorServicesAppDelegate.cpp
+++ b/Engine/Source/Programs/Mac/UE4EditorServices/Private/Mac/UE4EditorServicesAppDelegate.cpp
@@ -2,7 +2,7 @@
 
 #include "Core.h"
 #include "DesktopPlatformModule.h"
-#include "Runtime/Core/Public/Serialization/Json/Json.h"
+#include "Json.h"
 
 #import "UE4EditorServicesAppDelegate.h"
 
diff --git a/Engine/Source/Programs/Mac/UE4EditorServices/UE4EditorServices.Build.cs b/Engine/Source/Programs/Mac/UE4EditorServices/UE4EditorServices.Build.cs
--- a/Engine/Source/Programs/Mac/UE4EditorServices/UE4EditorServices.Build.cs
+++ b/Engine/Source/Programs/Mac/UE4EditorServices/UE4EditorServices.Build.cs
@@ -9,6 +9,13 @@ public class UE4EditorServices : ModuleRules
 		PublicIncludePaths.Add("Runtime/Launch/Public");
 		PublicIncludePaths.Add("Runtime/Launch/Private"); // Yuck. Required for RequiredProgramMainCPPInclude.h. (Also yuck).
 
-		PrivateDependencyModuleNames.AddRange(new string[] { "Core", "Projects", "DesktopPlatform" });
+		PrivateDependencyModuleNames.AddRange(
+			new string[] {
+				"Core",
+				"DesktopPlatform",
+				"Json",
+				"Projects",
+			}
+		);
 	}
 }
diff --git a/Engine/Source/Programs/UnrealHeaderTool/Private/Manifest.cpp b/Engine/Source/Programs/UnrealHeaderTool/Private/Manifest.cpp
--- a/Engine/Source/Programs/UnrealHeaderTool/Private/Manifest.cpp
+++ b/Engine/Source/Programs/UnrealHeaderTool/Private/Manifest.cpp
@@ -6,20 +6,23 @@
 #include "Json.h"
 #include "Manifest.h"
 
+
 namespace
 {
 	template <typename T> struct TJsonFieldType;
-	template <>           struct TJsonFieldType<double>                         { static const EJson::Type Value = EJson::Number;  };
-	template <>           struct TJsonFieldType<FString>                        { static const EJson::Type Value = EJson::String;  };
-	template <>           struct TJsonFieldType<bool>                           { static const EJson::Type Value = EJson::Boolean; };
-	template <>           struct TJsonFieldType<TArray<TSharedPtr<FJsonValue>>> { static const EJson::Type Value = EJson::Array;   };
-	template <>           struct TJsonFieldType<TSharedPtr<FJsonObject>>        { static const EJson::Type Value = EJson::Object;  };
+	template <>           struct TJsonFieldType<double>                         { static const EJson Value = EJson::Number;  };
+	template <>           struct TJsonFieldType<FString>                        { static const EJson Value = EJson::String;  };
+	template <>           struct TJsonFieldType<bool>                           { static const EJson Value = EJson::Boolean; };
+	template <>           struct TJsonFieldType<TArray<TSharedPtr<FJsonValue>>> { static const EJson Value = EJson::Array;   };
+	template <>           struct TJsonFieldType<TSharedPtr<FJsonObject>>        { static const EJson Value = EJson::Object;  };
 
 	template <typename T>
 	void GetJsonValue(T& OutVal, const TSharedPtr<FJsonValue>& JsonValue, const TCHAR* Outer)
 	{
 		if (JsonValue->Type != TJsonFieldType<T>::Value)
+		{
 			FError::Throwf(TEXT("'%s' is the wrong type"), Outer);
+		}
 
 		JsonValue->AsArgumentType(OutVal);
 	}
@@ -28,11 +31,16 @@ namespace
 	void GetJsonFieldValue(T& OutVal, const TSharedPtr<FJsonObject>& JsonObject, const TCHAR* FieldName, const TCHAR* Outer)
 	{
 		auto JsonValue = JsonObject->Values.Find(FieldName);
+
 		if (!JsonValue)
+		{
 			FError::Throwf(TEXT("Unable to find field '%s' in '%s'"), FieldName, Outer);
+		}
 
 		if ((*JsonValue)->Type != TJsonFieldType<T>::Value)
+		{
 			FError::Throwf(TEXT("Field '%s' in '%s' is the wrong type"), Outer);
+		}
 
 		(*JsonValue)->AsArgumentType(OutVal);
 	}
@@ -46,20 +54,25 @@ namespace
 	}
 }
 
+
 FManifest FManifest::LoadFromFile(const FString& Filename)
 {
 	FManifest Result;
-
 	FString FilenamePath = FPaths::GetPath(Filename);
-
 	FString Json;
+
 	if (!FFileHelper::LoadFileToString(Json, *Filename))
+	{
 		FError::Throwf(TEXT("Unable to load manifest: %s"), *Filename);
+	}
 
 	auto RootObject = TSharedPtr<FJsonObject>();
-	auto Reader     = TJsonReaderFactory<TCHAR>::Create(Json);
+	auto Reader = TJsonReaderFactory<TCHAR>::Create(Json);
+
 	if (!FJsonSerializer::Deserialize(Reader, RootObject))
+	{
 		FError::Throwf(TEXT("Manifest is malformed: %s"), *Filename);
+	}
 
 	TArray<TSharedPtr<FJsonValue>> ModulesArray;
 
@@ -80,10 +93,18 @@ FManifest FManifest::LoadFromFile(const FString& Filename)
 	Result.RootBuildPath = FPaths::ConvertRelativePathToFull(FilenamePath, Result.RootBuildPath);
 
 	// Ensure directories end with a slash, because this aids their use with FPaths::MakePathRelativeTo.
-	if (!Result.RootLocalPath.EndsWith(TEXT("/"))) { Result.RootLocalPath += TEXT("/"); }
-	if (!Result.RootBuildPath.EndsWith(TEXT("/"))) { Result.RootBuildPath += TEXT("/"); }
+	if (!Result.RootLocalPath.EndsWith(TEXT("/")))
+	{
+		Result.RootLocalPath += TEXT("/");
+	}
+
+	if (!Result.RootBuildPath.EndsWith(TEXT("/")))
+	{
+		Result.RootBuildPath += TEXT("/");
+	}
 
 	int32 ModuleIndex = 0;
+
 	for (auto Module : ModulesArray)
 	{
 		auto ModuleObj = Module->AsObject();
diff --git a/Engine/Source/Programs/UnrealHeaderTool/UnrealHeaderTool.Build.cs b/Engine/Source/Programs/UnrealHeaderTool/UnrealHeaderTool.Build.cs
--- a/Engine/Source/Programs/UnrealHeaderTool/UnrealHeaderTool.Build.cs
+++ b/Engine/Source/Programs/UnrealHeaderTool/UnrealHeaderTool.Build.cs
@@ -8,7 +8,15 @@ public class UnrealHeaderTool : ModuleRules
 	{
 		PublicIncludePaths.Add("Runtime/Launch/Public");
 
-		PrivateDependencyModuleNames.AddRange(new string[] { "Core", "CoreUObject", "Projects" });
+		PrivateDependencyModuleNames.AddRange(
+			new string[]
+			{
+				"Core",
+				"CoreUObject",
+				"Json",
+				"Projects",
+			}
+		);
 	
 		PrivateIncludePaths.Add("Runtime/Launch/Private");		// For LaunchEngineLoop.cpp include
 		
diff --git a/Engine/Source/Runtime/Analytics/AnalyticsET/AnalyticsET.Build.cs b/Engine/Source/Runtime/Analytics/AnalyticsET/AnalyticsET.Build.cs
--- a/Engine/Source/Runtime/Analytics/AnalyticsET/AnalyticsET.Build.cs
+++ b/Engine/Source/Runtime/Analytics/AnalyticsET/AnalyticsET.Build.cs
@@ -31,7 +31,8 @@ namespace UnrealBuildTool.Rules
 				new string[]
 				{
 					"Analytics",
-					"HTTP"
+					"HTTP",
+					"Json",
                     // ... add private dependencies that you statically link with here ...
 				}
 				);
diff --git a/Engine/Source/Runtime/Analytics/AnalyticsMulticast/AnalyticsMulticast.Build.cs b/Engine/Source/Runtime/Analytics/AnalyticsMulticast/AnalyticsMulticast.Build.cs
--- a/Engine/Source/Runtime/Analytics/AnalyticsMulticast/AnalyticsMulticast.Build.cs
+++ b/Engine/Source/Runtime/Analytics/AnalyticsMulticast/AnalyticsMulticast.Build.cs
@@ -31,6 +31,7 @@ namespace UnrealBuildTool.Rules
 				new string[]
 				{
 					"Analytics",
+					"Json",
                     // ... add private dependencies that you statically link with here ...
 				}
 				);
diff --git a/Engine/Source/Runtime/Analytics/AnalyticsSwrve/AnalyticsSwrve.Build.cs b/Engine/Source/Runtime/Analytics/AnalyticsSwrve/AnalyticsSwrve.Build.cs
--- a/Engine/Source/Runtime/Analytics/AnalyticsSwrve/AnalyticsSwrve.Build.cs
+++ b/Engine/Source/Runtime/Analytics/AnalyticsSwrve/AnalyticsSwrve.Build.cs
@@ -12,9 +12,10 @@ namespace UnrealBuildTool.Rules
 				new string[]
 				{
 					"Analytics",
-					"HTTP"
+					"HTTP",
+					"Json",
 				}
-				);
+			);
 		}
 	}
 }
diff --git a/Engine/Source/Runtime/Core/Private/Internationalization/InternationalizationArchiveJsonSerializer.cpp b/Engine/Source/Runtime/Core/Private/Internationalization/InternationalizationArchiveJsonSerializer.cpp
deleted file mode 100644
diff --git a/Engine/Source/Runtime/Core/Private/Internationalization/InternationalizationManifestJsonSerializer.cpp b/Engine/Source/Runtime/Core/Private/Internationalization/InternationalizationManifestJsonSerializer.cpp
deleted file mode 100644
diff --git a/Engine/Source/Runtime/Core/Private/Internationalization/InternationalizationMetadataJsonSerializer.cpp b/Engine/Source/Runtime/Core/Private/Internationalization/InternationalizationMetadataJsonSerializer.cpp
deleted file mode 100644
diff --git a/Engine/Source/Runtime/Core/Private/Internationalization/TextLocalizationManager.cpp b/Engine/Source/Runtime/Core/Private/Internationalization/TextLocalizationManager.cpp
--- a/Engine/Source/Runtime/Core/Private/Internationalization/TextLocalizationManager.cpp
+++ b/Engine/Source/Runtime/Core/Private/Internationalization/TextLocalizationManager.cpp
@@ -2,10 +2,8 @@
 
 #include "Core.h"
 #include "InternationalizationManifest.h"
-#include "Json.h"
-#include "JsonDocumentObjectModel.h"
 #include "TextLocalizationResourceGenerator.h"
-#include "InternationalizationManifestJsonSerializer.h"
+
 
 DEFINE_LOG_CATEGORY_STATIC(LogTextLocalizationManager, Log, All);
 
@@ -653,34 +651,8 @@ void FTextLocalizationManager::FindKeyNamespaceFromDisplayString(TSharedRef<FStr
 	}
 }
 
-namespace
-{
-	TSharedPtr<FJsonObject> ReadJSONTextFile(const FString& InFilePath)
-	{
-		//read in file as string
-		FString FileContents;
-		if ( !FFileHelper::LoadFileToString(FileContents, *InFilePath) )
-		{
-			UE_LOG(LogTextLocalizationManager, Error,TEXT("Failed to load file %s."), *InFilePath);
-			return NULL;
-		}
-
-		//parse as JSON
-		TSharedPtr<FJsonObject> JSONObject;
 
-		TSharedRef< TJsonReader<> > Reader = TJsonReaderFactory<>::Create( FileContents );
-
-		if( !FJsonSerializer::Deserialize( Reader, JSONObject ) || !JSONObject.IsValid())
-		{
-			UE_LOG(LogTextLocalizationManager, Error,TEXT("Invalid JSON in file %s."), *InFilePath);
-			return NULL;
-		}
-
-		return JSONObject;
-	}
-}
-
-void FTextLocalizationManager::RegenerateResources(const FString& ConfigFilePath)
+void FTextLocalizationManager::RegenerateResources( const FString& ConfigFilePath, IInternationalizationArchiveSerializer& ArchiveSerializer, IInternationalizationManifestSerializer& ManifestSerializer )
 {
 	// Add one to the revision index, so all FText's refresh.
 	++HeadCultureRevisionIndex;
@@ -765,20 +737,21 @@ void FTextLocalizationManager::RegenerateResources(const FString& ConfigFilePath
 		// Read the manifest file from the source path.
 		FString ManifestFilePath = (SourcePath / ManifestName);
 		ManifestFilePath = FPaths::ConvertRelativePathToFull(ManifestFilePath);
-		TSharedPtr<FJsonObject> ManifestJSONObject = ReadJSONTextFile(ManifestFilePath);
-		if( !(ManifestJSONObject.IsValid()) )
+		FArchive* ManifestFile = IFileManager::Get().CreateFileReader(*ManifestFilePath);
+
+		if (ManifestFile == nullptr)
 		{
 			UE_LOG(LogTextLocalizationManager, Error, TEXT("No manifest found at %s."), *ManifestFilePath);
 			return;
 		}
+
 		TSharedRef<FInternationalizationManifest> InternationalizationManifest = MakeShareable( new FInternationalizationManifest );
 		{
-			FInternationalizationManifestJsonSerializer ManifestSerializer;
-			ManifestSerializer.DeserializeManifest(ManifestJSONObject.ToSharedRef(), InternationalizationManifest);
+			ManifestSerializer.DeserializeManifest(*ManifestFile, InternationalizationManifest);
 		}
 
 		// Write resource.
-		FTextLocalizationResourceGenerator::Generate(SourcePath, InternationalizationManifest, LocaleNames[i], &(MemoryWriter));
+		FTextLocalizationResourceGenerator::Generate(SourcePath, InternationalizationManifest, LocaleNames[i], &(MemoryWriter), ArchiveSerializer);
 
 		MemoryWriter.Close();
 	}
diff --git a/Engine/Source/Runtime/Core/Private/Internationalization/TextLocalizationResourceGenerator.cpp b/Engine/Source/Runtime/Core/Private/Internationalization/TextLocalizationResourceGenerator.cpp
--- a/Engine/Source/Runtime/Core/Private/Internationalization/TextLocalizationResourceGenerator.cpp
+++ b/Engine/Source/Runtime/Core/Private/Internationalization/TextLocalizationResourceGenerator.cpp
@@ -2,36 +2,12 @@
 
 #include "Core.h"
 #include "InternationalizationManifest.h"
-#include "InternationalizationManifestJsonSerializer.h"
 #include "InternationalizationArchive.h"
-#include "InternationalizationArchiveJsonSerializer.h"
 #include "TextLocalizationResourceGenerator.h"
 
-DEFINE_LOG_CATEGORY_STATIC(LogTextLocalizationResourceGenerator, Log, All);
-
-TSharedPtr<FJsonObject> FTextLocalizationResourceGenerator::ReadJSONTextFile(const FString& InFilePath)
-{
-	//read in file as string
-	FString FileContents;
-	if ( !FFileHelper::LoadFileToString(FileContents, *InFilePath) )
-	{
-		UE_LOG(LogTextLocalizationResourceGenerator, Error,TEXT("Failed to load file %s."), *InFilePath);
-		return NULL;
-	}
-
-	//parse as JSON
-	TSharedPtr<FJsonObject> JSONObject;
-
-	TSharedRef< TJsonReader<> > Reader = TJsonReaderFactory<>::Create( FileContents );
 
-	if( !FJsonSerializer::Deserialize( Reader, JSONObject ) || !JSONObject.IsValid())
-	{
-		UE_LOG(LogTextLocalizationResourceGenerator, Error,TEXT("Invalid JSON in file %s."), *InFilePath);
-		return NULL;
-	}
+DEFINE_LOG_CATEGORY_STATIC(LogTextLocalizationResourceGenerator, Log, All);
 
-	return JSONObject;
-}
 
 void FTextLocalizationResourceGenerator::FLocalizationEntryTracker::ReportCollisions() const
 {
@@ -164,7 +140,7 @@ bool FTextLocalizationResourceGenerator::FLocalizationEntryTracker::WriteToArchi
 	return WasSuccessful;
 }
 
-bool FTextLocalizationResourceGenerator::Generate(const FString& SourcePath, const TSharedRef<FInternationalizationManifest>& InternationalizationManifest, const FString& CultureToGenerate, FArchive* const DestinationArchive)
+bool FTextLocalizationResourceGenerator::Generate(const FString& SourcePath, const TSharedRef<FInternationalizationManifest>& InternationalizationManifest, const FString& CultureToGenerate, FArchive* const DestinationArchive, IInternationalizationArchiveSerializer& ArchiveSerializer)
 {
 	FLocalizationEntryTracker LocalizationEntryTracker;
 
@@ -188,16 +164,17 @@ bool FTextLocalizationResourceGenerator::Generate(const FString& SourcePath, con
 		// Read each archive file from the culture-named directory in the source path.
 		FString ArchiveFilePath = CulturePath / ArchiveName;
 		ArchiveFilePath = FPaths::ConvertRelativePathToFull(ArchiveFilePath);
-		TSharedPtr<FJsonObject> ArchiveJSONObject = ReadJSONTextFile(ArchiveFilePath);
-		if( !(ArchiveJSONObject.IsValid()) )
+		FArchive* ArchiveFile = IFileManager::Get().CreateFileReader(*ArchiveFilePath);
+
+		if (ArchiveFile == nullptr)
 		{
 			UE_LOG(LogTextLocalizationResourceGenerator, Error, TEXT("No archive found at %s."), *ArchiveFilePath);
 			continue;
 		}
+
 		TSharedRef<FInternationalizationArchive> InternationalizationArchive = MakeShareable( new FInternationalizationArchive );
 		{
-			FInternationalizationArchiveJsonSerializer ArchiveSerializer;
-			ArchiveSerializer.DeserializeArchive(ArchiveJSONObject.ToSharedRef(), InternationalizationArchive);
+			ArchiveSerializer.DeserializeArchive(*ArchiveFile, InternationalizationArchive);
 		}
 
 		// Generate text localization resource from manifest and archive entries.
diff --git a/Engine/Source/Runtime/Core/Private/Serialization/Json/Json.cpp b/Engine/Source/Runtime/Core/Private/Serialization/Json/Json.cpp
deleted file mode 100644
diff --git a/Engine/Source/Runtime/Core/Private/Serialization/Json/JsonDocumentObjectModel.cpp b/Engine/Source/Runtime/Core/Private/Serialization/Json/JsonDocumentObjectModel.cpp
deleted file mode 100644
diff --git a/Engine/Source/Runtime/Core/Public/Internationalization/InternationalizationArchive.h b/Engine/Source/Runtime/Core/Public/Internationalization/InternationalizationArchive.h
--- a/Engine/Source/Runtime/Core/Public/Internationalization/InternationalizationArchive.h
+++ b/Engine/Source/Runtime/Core/Public/Internationalization/InternationalizationArchive.h
@@ -37,7 +37,7 @@ typedef TMultiMap< FString, TSharedRef< FArchiveEntry >, FDefaultSetAllocator, F
 class CORE_API FInternationalizationArchive 
 {
 	friend class IInternationalizationArchiveSerializer;
-	friend class FInternationalizationArchiveJsonSerializer;
+
 public:
 	enum EFormatVersion
 	{
@@ -49,10 +49,8 @@ public:
 	};
 
 	FInternationalizationArchive()
-		: FormatVersion(static_cast<int32>(EFormatVersion::Latest))
-	{
-
-	}
+		: FormatVersion(EFormatVersion::Latest)
+	{ }
 
 	TArchiveEntryContainer::TConstIterator GetEntryIterator() const;
 
@@ -72,7 +70,7 @@ public:
 
 	void SetFormatVersion(const EFormatVersion Version)
 	{
-		FormatVersion = static_cast<int>(Version);
+		FormatVersion = Version;
 	}
 
 	EFormatVersion GetFormatVersion() const
@@ -81,6 +79,7 @@ public:
 	}
 
 private:
-	int32 FormatVersion;
+
+	EFormatVersion FormatVersion;
 	TArchiveEntryContainer EntriesBySourceText;
-};
\ No newline at end of file
+};
diff --git a/Engine/Source/Runtime/Core/Public/Internationalization/InternationalizationArchiveJsonSerializer.h b/Engine/Source/Runtime/Core/Public/Internationalization/InternationalizationArchiveJsonSerializer.h
deleted file mode 100644
diff --git a/Engine/Source/Runtime/Core/Public/Internationalization/InternationalizationManifest.h b/Engine/Source/Runtime/Core/Public/Internationalization/InternationalizationManifest.h
--- a/Engine/Source/Runtime/Core/Public/Internationalization/InternationalizationManifest.h
+++ b/Engine/Source/Runtime/Core/Public/Internationalization/InternationalizationManifest.h
@@ -1,7 +1,9 @@
 // Copyright 1998-2014 Epic Games, Inc. All Rights Reserved.
 #pragma once
 
-;class FLocMetadataObject;
+
+class FLocMetadataObject;
+
 
 struct CORE_API FContext
 {
@@ -18,9 +20,8 @@ public:
 	inline bool operator!=(const FContext& Other) const { return !(*this == Other); }
 	bool operator<(const FContext& Other) const;
 
-
-
 public:
+
 	FString Key;
 	FString SourceLocation;
 	bool bIsOptional;
@@ -37,8 +38,7 @@ public:
 	FLocItem( const FString& InSourceText )
 		: Text( InSourceText )
 		, MetadataObj( NULL )
-	{
-	}
+	{ }
 
 	/** Copy ctor */
 	FLocItem( const FLocItem& Other );
@@ -55,6 +55,7 @@ public:
 	TSharedPtr<FLocMetadataObject> MetadataObj;
 };
 
+
 class FManifestEntry
 {
 public:
@@ -172,7 +173,6 @@ public:
 	}
 
 	friend class IInternationalizationManifestSerializer;
-	friend class FInternationalizationManifestJsonSerializer;
 
 private:
 
@@ -180,6 +180,3 @@ private:
 	TManifestEntryBySourceTextContainer EntriesBySourceText;
 	TManifestEntryByContextIdContainer EntriesByContextId;
 };
-
-
-
diff --git a/Engine/Source/Runtime/Core/Public/Internationalization/InternationalizationManifestJsonSerializer.h b/Engine/Source/Runtime/Core/Public/Internationalization/InternationalizationManifestJsonSerializer.h
deleted file mode 100644
diff --git a/Engine/Source/Runtime/Core/Public/Internationalization/InternationalizationMetadataJsonSerializer.h b/Engine/Source/Runtime/Core/Public/Internationalization/InternationalizationMetadataJsonSerializer.h
deleted file mode 100644
diff --git a/Engine/Source/Runtime/Core/Public/Internationalization/TextLocalizationManager.h b/Engine/Source/Runtime/Core/Public/Internationalization/TextLocalizationManager.h
--- a/Engine/Source/Runtime/Core/Public/Internationalization/TextLocalizationManager.h
+++ b/Engine/Source/Runtime/Core/Public/Internationalization/TextLocalizationManager.h
@@ -2,6 +2,10 @@
 
 #pragma once
 
+#include "IInternationalizationArchiveSerializer.h"
+#include "IInternationalizationManifestSerializer.h"
+
+
 namespace ELocalizationResourceSource
 {
 	enum Type
@@ -87,7 +91,7 @@ public:
 	 */
 	void CORE_API FindKeyNamespaceFromDisplayString(TSharedRef<FString, ESPMode::ThreadSafe> InDisplayString, TSharedPtr<FString, ESPMode::ThreadSafe>& OutNamespace, TSharedPtr<FString, ESPMode::ThreadSafe>& OutKey);
 
-	void CORE_API RegenerateResources(const FString& ConfigFilePath);
+	void CORE_API RegenerateResources( const FString& ConfigFilePath, IInternationalizationArchiveSerializer& ArchiveSerializer, IInternationalizationManifestSerializer& ManifestSerializer );
 
 	/** Returns the current culture revision index */
 	int CORE_API GetHeadCultureRevision() const { return HeadCultureRevisionIndex; }
diff --git a/Engine/Source/Runtime/Core/Public/Internationalization/TextLocalizationResourceGenerator.h b/Engine/Source/Runtime/Core/Public/Internationalization/TextLocalizationResourceGenerator.h
--- a/Engine/Source/Runtime/Core/Public/Internationalization/TextLocalizationResourceGenerator.h
+++ b/Engine/Source/Runtime/Core/Public/Internationalization/TextLocalizationResourceGenerator.h
@@ -5,12 +5,10 @@
 class FTextLocalizationResourceGenerator
 {
 public:
-	CORE_API static bool Generate(const FString& SourcePath, const TSharedRef<FInternationalizationManifest>& InternationalizationManifest, const FString& CultureToGenerate, FArchive* const DestinationArchive);
+	CORE_API static bool Generate(const FString& SourcePath, const TSharedRef<FInternationalizationManifest>& InternationalizationManifest, const FString& CultureToGenerate, FArchive* const DestinationArchive, IInternationalizationArchiveSerializer& ArchiveSerializer);
 
 private:
-	static TSharedPtr<FJsonObject> ReadJSONTextFile(const FString& InFilePath);
 
-private:
 	struct FLocalizationEntryTracker
 	{
 		struct FEntry
diff --git a/Engine/Source/Runtime/Core/Public/Serialization/Json/Json.h b/Engine/Source/Runtime/Core/Public/Serialization/Json/Json.h
deleted file mode 100644
diff --git a/Engine/Source/Runtime/Core/Public/Serialization/Json/JsonDocumentObjectModel.h b/Engine/Source/Runtime/Core/Public/Serialization/Json/JsonDocumentObjectModel.h
deleted file mode 100644
diff --git a/Engine/Source/Runtime/Core/Public/Serialization/Json/JsonPrintPolicies.h b/Engine/Source/Runtime/Core/Public/Serialization/Json/JsonPrintPolicies.h
deleted file mode 100644
diff --git a/Engine/Source/Runtime/Core/Public/Serialization/Json/JsonReader.h b/Engine/Source/Runtime/Core/Public/Serialization/Json/JsonReader.h
deleted file mode 100644
diff --git a/Engine/Source/Runtime/Core/Public/Serialization/Json/JsonSerializer.h b/Engine/Source/Runtime/Core/Public/Serialization/Json/JsonSerializer.h
deleted file mode 100644
diff --git a/Engine/Source/Runtime/Core/Public/Serialization/Json/JsonWriter.h b/Engine/Source/Runtime/Core/Public/Serialization/Json/JsonWriter.h
deleted file mode 100644
diff --git a/Engine/Source/Runtime/Engine/Engine.Build.cs b/Engine/Source/Runtime/Engine/Engine.Build.cs
--- a/Engine/Source/Runtime/Engine/Engine.Build.cs
+++ b/Engine/Source/Runtime/Engine/Engine.Build.cs
@@ -42,6 +42,7 @@ public class Engine : ModuleRules
 			new string[] { 
 				"Core",
 				"CoreUObject",
+				"Json",
 				"Slate",
 				"InputCore",
 				"Messaging",
diff --git a/Engine/Source/Runtime/Internationalization/Internationalization.Build.cs b/Engine/Source/Runtime/Internationalization/Internationalization.Build.cs
new file mode 100644
--- /dev/null
+++ b/Engine/Source/Runtime/Internationalization/Internationalization.Build.cs
@@ -0,0 +1,30 @@
+// Copyright 1998-2014 Epic Games, Inc. All Rights Reserved.
+
+namespace UnrealBuildTool.Rules
+{
+	public class Internationalization : ModuleRules
+	{
+		public Internationalization(TargetInfo Target)
+		{
+			PublicDependencyModuleNames.AddRange(
+				new string[]
+				{
+					"Core",
+				}
+			); 
+
+			PrivateDependencyModuleNames.AddRange(
+				new string[]
+				{
+					"Json",
+				}
+			); 
+
+			PrivateIncludePaths.AddRange(
+				new string[] {
+					"Runtime/Internationalization/Private",
+				}
+			);
+		}
+	}
+}
diff --git a/Engine/Source/Runtime/Internationalization/Private/InternationalizationModule.cpp b/Engine/Source/Runtime/Internationalization/Private/InternationalizationModule.cpp
new file mode 100644
--- /dev/null
+++ b/Engine/Source/Runtime/Internationalization/Private/InternationalizationModule.cpp
@@ -0,0 +1,27 @@
+// Copyright 1998-2014 Epic Games, Inc. All Rights Reserved.
+
+#include "InternationalizationPrivatePCH.h"
+#include "ModuleManager.h"
+
+
+/**
+ * Implements the Internationalization module.
+ */
+class FInternationalizationModule
+	: public IModuleInterface
+{
+public:
+
+	// IModuleInterface interface
+
+	virtual void StartupModule( ) override { }
+	virtual void ShutdownModule( ) override { }
+
+	virtual bool SupportsDynamicReloading( ) override
+	{
+		return false;
+	}
+};
+
+
+IMPLEMENT_MODULE(FInternationalizationModule, Internationalization);
diff --git a/Engine/Source/Runtime/Internationalization/Private/InternationalizationPrivatePCH.h b/Engine/Source/Runtime/Internationalization/Private/InternationalizationPrivatePCH.h
new file mode 100644
--- /dev/null
+++ b/Engine/Source/Runtime/Internationalization/Private/InternationalizationPrivatePCH.h
@@ -0,0 +1,11 @@
+// Copyright 1998-2014 Epic Games, Inc. All Rights Reserved.
+
+#pragma once
+
+#include "Internationalization.h"
+
+
+/* Private dependencies
+ *****************************************************************************/
+
+#include "Internationalization.h"
diff --git a/Engine/Source/Runtime/Internationalization/Private/Serialization/JsonInternationalizationArchiveSerializer.cpp b/Engine/Source/Runtime/Internationalization/Private/Serialization/JsonInternationalizationArchiveSerializer.cpp
new file mode 100644
--- /dev/null
+++ b/Engine/Source/Runtime/Internationalization/Private/Serialization/JsonInternationalizationArchiveSerializer.cpp
@@ -0,0 +1,468 @@
+﻿// Copyright 1998-2014 Epic Games, Inc. All Rights Reserved.
+
+#include "InternationalizationPrivatePCH.h"
+#include "Internationalization/InternationalizationArchive.h"
+#include "Internationalization/InternationalizationMetadata.h"
+#include "JsonInternationalizationArchiveSerializer.h"
+#include "JsonInternationalizationMetadataSerializer.h"
+
+
+DEFINE_LOG_CATEGORY_STATIC(LogInternationalizationArchiveSerializer, Log, All);
+
+
+const FString FJsonInternationalizationArchiveSerializer::TAG_FORMATVERSION = TEXT("FormatVersion");
+const FString FJsonInternationalizationArchiveSerializer::TAG_NAMESPACE = TEXT("Namespace");
+const FString FJsonInternationalizationArchiveSerializer::TAG_CHILDREN = TEXT("Children");
+const FString FJsonInternationalizationArchiveSerializer::TAG_SUBNAMESPACES = TEXT("Subnamespaces");
+const FString FJsonInternationalizationArchiveSerializer::TAG_DEPRECATED_DEFAULTTEXT = TEXT("DefaultText");
+const FString FJsonInternationalizationArchiveSerializer::TAG_DEPRECATED_TRANSLATEDTEXT = TEXT("TranslatedText");
+const FString FJsonInternationalizationArchiveSerializer::TAG_OPTIONAL = TEXT("Optional");
+const FString FJsonInternationalizationArchiveSerializer::TAG_SOURCE = TEXT("Source");
+const FString FJsonInternationalizationArchiveSerializer::TAG_SOURCE_TEXT = TEXT("Text");
+const FString FJsonInternationalizationArchiveSerializer::TAG_TRANSLATION = TEXT("Translation");
+const FString FJsonInternationalizationArchiveSerializer::TAG_TRANSLATION_TEXT = FJsonInternationalizationArchiveSerializer::TAG_SOURCE_TEXT;
+const FString FJsonInternationalizationArchiveSerializer::TAG_METADATA_KEY = TEXT("Key");
+const FString FJsonInternationalizationArchiveSerializer::NAMESPACE_DELIMITER = TEXT(".");
+
+
+struct FCompareArchiveEntryBySourceAndKey
+{
+	FORCEINLINE bool operator()( TSharedPtr< FArchiveEntry > A, TSharedPtr< FArchiveEntry > B ) const
+	{
+		bool bResult = false;
+		if( A->Source < B->Source )
+		{
+			bResult = true;
+		}
+		else if( A->Source == B->Source )
+		{
+			if( A->KeyMetadataObj.IsValid() != B->KeyMetadataObj.IsValid() )
+			{
+				bResult = B->KeyMetadataObj.IsValid();
+			}
+			else if( A->KeyMetadataObj.IsValid() && B->KeyMetadataObj.IsValid() )
+			{
+				bResult = (*(A->KeyMetadataObj) < *(B->KeyMetadataObj));
+			}
+		}
+		return bResult;
+	}
+};
+
+
+struct FCompareStructuredArchiveEntryByNamespace
+{
+	FORCEINLINE bool operator()( TSharedPtr< FStructuredArchiveEntry > A, TSharedPtr< FStructuredArchiveEntry > B ) const
+	{
+		return A->Namespace < B->Namespace;
+	}
+};
+
+
+bool FJsonInternationalizationArchiveSerializer::DeserializeArchive( FArchive& Archive, TSharedRef< FInternationalizationArchive > InternationalizationArchive )
+{
+	TSharedPtr< FJsonObject > JsonArchiveObj;
+	TSharedRef< TJsonReader<> > Reader = TJsonReaderFactory<>::Create( &Archive );
+	bool bExecSuccessful = FJsonSerializer::Deserialize( Reader, JsonArchiveObj );
+
+	if (bExecSuccessful && JsonArchiveObj.IsValid())
+	{
+		bExecSuccessful = DeserializeInternal( JsonArchiveObj.ToSharedRef(), InternationalizationArchive );
+	}		
+
+	return bExecSuccessful;
+}
+
+
+bool FJsonInternationalizationArchiveSerializer::DeserializeArchive( const FString& InStr, TSharedRef< FInternationalizationArchive > InternationalizationArchive )
+{
+	TSharedPtr< FJsonObject > JsonArchiveObj;
+	TSharedRef< TJsonReader<> > Reader = TJsonReaderFactory<>::Create( InStr );
+	bool bExecSuccessful = FJsonSerializer::Deserialize( Reader, JsonArchiveObj );
+
+	if ( bExecSuccessful && JsonArchiveObj.IsValid() )
+	{
+		bExecSuccessful = DeserializeInternal( JsonArchiveObj.ToSharedRef(), InternationalizationArchive );
+	}
+
+	return bExecSuccessful;
+}
+
+
+bool FJsonInternationalizationArchiveSerializer::DeserializeArchive( TSharedRef< FJsonObject > InJsonObj, TSharedRef< FInternationalizationArchive > InternationalizationArchive )
+{
+	return DeserializeInternal( InJsonObj, InternationalizationArchive );
+}
+
+
+bool FJsonInternationalizationArchiveSerializer::SerializeArchive( TSharedRef< const FInternationalizationArchive > InternationalizationArchive, FArchive& Archive )
+{
+	TSharedRef< FJsonObject > JsonArchiveObj = MakeShareable( new FJsonObject );
+	bool bExecSuccessful = SerializeInternal( InternationalizationArchive, JsonArchiveObj );
+
+	if( bExecSuccessful )
+	{
+		TSharedRef< TJsonWriter<> > Writer = TJsonWriterFactory<>::Create( &Archive );
+		bExecSuccessful = FJsonSerializer::Serialize( JsonArchiveObj, Writer );
+		Writer->Close();
+	}
+	return bExecSuccessful;
+}
+
+
+bool FJsonInternationalizationArchiveSerializer::SerializeArchive( TSharedRef< const FInternationalizationArchive > InternationalizationArchive, FString& Str )
+{
+	TSharedRef< FJsonObject > JsonArchiveObj = MakeShareable( new FJsonObject );
+	bool bExecSuccessful = SerializeInternal( InternationalizationArchive, JsonArchiveObj );
+
+	if( bExecSuccessful )
+	{
+		TSharedRef< TJsonWriter<> > Writer = TJsonWriterFactory<>::Create( &Str );
+		bExecSuccessful = FJsonSerializer::Serialize( JsonArchiveObj, Writer );
+		Writer->Close();
+	}
+	return bExecSuccessful;
+}
+
+
+bool FJsonInternationalizationArchiveSerializer::SerializeArchive( TSharedRef< const FInternationalizationArchive > InternationalizationArchive, TSharedRef< FJsonObject > JsonObj )
+{
+	return SerializeInternal( InternationalizationArchive, JsonObj );
+}
+
+
+bool FJsonInternationalizationArchiveSerializer::DeserializeInternal( TSharedRef< FJsonObject > InJsonObj, TSharedRef< FInternationalizationArchive > InternationalizationArchive )
+{
+	if( InJsonObj->HasField( TAG_FORMATVERSION ) )
+	{
+		const int32 FormatVersion = static_cast<int>(InJsonObj->GetNumberField( TAG_FORMATVERSION ));
+		InternationalizationArchive->SetFormatVersion(static_cast<FInternationalizationArchive::EFormatVersion>(FormatVersion));
+	}
+	else
+	{
+		InternationalizationArchive->SetFormatVersion(FInternationalizationArchive::EFormatVersion::Initial);
+	}
+
+	return JsonObjToArchive( InJsonObj, TEXT(""), InternationalizationArchive );
+}
+
+
+bool FJsonInternationalizationArchiveSerializer::SerializeInternal( TSharedRef< const FInternationalizationArchive > InInternationalizationArchive, TSharedRef< FJsonObject > JsonObj )
+{
+	TSharedPtr< FStructuredArchiveEntry > RootElement = MakeShareable( new FStructuredArchiveEntry( TEXT("") ) );
+
+	// Condition the data so that it exists in a structured hierarchy for easy population of the JSON object.
+	GenerateStructuredData( InInternationalizationArchive, RootElement );
+
+	SortStructuredData( RootElement );
+
+	// Clear anything that may be in the JSON object
+	JsonObj->Values.Empty();
+
+	// Set format version.
+	JsonObj->SetNumberField(TAG_FORMATVERSION, static_cast<double>(InInternationalizationArchive->GetFormatVersion()));
+
+	// Setup the JSON object using the structured data created
+	StructuredDataToJsonObj( RootElement, JsonObj );
+	return true;
+}
+
+
+bool FJsonInternationalizationArchiveSerializer::JsonObjToArchive( TSharedRef< FJsonObject > InJsonObj, FString ParentNamespace, TSharedRef< FInternationalizationArchive > InternationalizationArchive )
+{
+	bool bConvertSuccess = true;
+	FString AccumulatedNamespace = ParentNamespace;
+	
+	if( InJsonObj->HasField( TAG_NAMESPACE) )
+	{
+		if( !( AccumulatedNamespace.IsEmpty() ) )
+		{
+			AccumulatedNamespace += NAMESPACE_DELIMITER;
+		}
+		AccumulatedNamespace += InJsonObj->GetStringField( TAG_NAMESPACE );
+	}
+	else
+	{
+		UE_LOG( LogInternationalizationArchiveSerializer, Warning,TEXT("Encountered an object with a missing namespace while converting to Internationalization archive.") );
+		bConvertSuccess = false;
+	}
+
+	// Process all the child objects
+	if( bConvertSuccess && InJsonObj->HasField( TAG_CHILDREN ) )
+	{
+		const TArray< TSharedPtr<FJsonValue> > ChildrenArray = InJsonObj->GetArrayField( TAG_CHILDREN );
+
+		for( TArray< TSharedPtr< FJsonValue > >::TConstIterator ChildIter( ChildrenArray.CreateConstIterator() ); ChildIter; ++ChildIter )
+		{
+			const TSharedPtr< FJsonValue >  ChildEntry = *ChildIter;
+			const TSharedPtr< FJsonObject > ChildJSONObject = ChildEntry->AsObject();
+
+			FString SourceText;
+			TSharedPtr< FLocMetadataObject > SourceMetadata;
+			if( ChildJSONObject->HasTypedField< EJson::String >( TAG_DEPRECATED_DEFAULTTEXT ) )
+			{
+				SourceText = ChildJSONObject->GetStringField( TAG_DEPRECATED_DEFAULTTEXT );
+			} 
+			else if( ChildJSONObject->HasTypedField< EJson::Object >( TAG_SOURCE ) )
+			{
+				const TSharedPtr< FJsonObject > SourceJSONObject = ChildJSONObject->GetObjectField( TAG_SOURCE );
+				if( SourceJSONObject->HasTypedField< EJson::String >( TAG_SOURCE_TEXT ) )
+				{
+					SourceText = SourceJSONObject->GetStringField( TAG_SOURCE_TEXT );
+
+					// Source meta data is mixed in with the source text, we'll process metadata if the source json object has more than one entry
+					if( SourceJSONObject->Values.Num() > 1 )
+					{
+						// We load in the entire source object as metadata and just remove the source text.
+						FJsonInternationalizationMetaDataSerializer::DeserializeMetadata( SourceJSONObject.ToSharedRef(), SourceMetadata );
+						if( SourceMetadata.IsValid() )
+						{
+							SourceMetadata->Values.Remove( TAG_SOURCE_TEXT );
+						}
+					}
+				}
+				else
+				{
+					bConvertSuccess = false;
+				}
+			}
+			else
+			{
+				bConvertSuccess = false;
+			}
+
+			FString TranslationText;
+			TSharedPtr< FLocMetadataObject > TranslationMetadata;
+			if( ChildJSONObject->HasTypedField< EJson::String >( TAG_DEPRECATED_TRANSLATEDTEXT ) )
+			{
+				TranslationText = ChildJSONObject->GetStringField( TAG_DEPRECATED_TRANSLATEDTEXT );
+			} 
+			else if( ChildJSONObject->HasTypedField< EJson::Object >( TAG_TRANSLATION ) )
+			{
+				const TSharedPtr< FJsonObject > TranslationJSONObject = ChildJSONObject->GetObjectField( TAG_TRANSLATION );
+				if( TranslationJSONObject->HasTypedField< EJson::String >( TAG_TRANSLATION_TEXT ) )
+				{
+					TranslationText = TranslationJSONObject->GetStringField( TAG_TRANSLATION_TEXT );
+
+					// Source meta data is mixed in with the source text, we'll process metadata if the source json object has more than one entry
+					if( TranslationJSONObject->Values.Num() > 1 )
+					{
+						// We load in the entire source object as metadata and remove the source text
+						FJsonInternationalizationMetaDataSerializer::DeserializeMetadata( TranslationJSONObject.ToSharedRef(), TranslationMetadata );
+						if( TranslationJSONObject.IsValid() )
+						{
+							TranslationJSONObject->Values.Remove( TAG_TRANSLATION_TEXT );
+						}
+					}
+				}
+				else
+				{
+					bConvertSuccess = false;
+				}
+			}
+			else
+			{
+				bConvertSuccess = false;
+			}
+
+			if( bConvertSuccess )
+			{
+				FLocItem Source( SourceText );
+				Source.MetadataObj = SourceMetadata;
+
+				FLocItem Translation( TranslationText );
+				Translation.MetadataObj = TranslationMetadata;
+
+				bool bIsOptional = false;
+				if( ChildJSONObject->HasTypedField< EJson::Boolean >( TAG_OPTIONAL ) )
+				{
+					bIsOptional = ChildJSONObject->GetBoolField( TAG_OPTIONAL );
+				}
+
+				TSharedPtr< FLocMetadataObject > MetadataNode;
+				if( ChildJSONObject->HasTypedField< EJson::Object >( TAG_METADATA_KEY ) )
+				{
+					const TSharedPtr< FJsonObject > MetaDataKeyJSONObject = ChildJSONObject->GetObjectField( TAG_METADATA_KEY );
+					FJsonInternationalizationMetaDataSerializer::DeserializeMetadata( MetaDataKeyJSONObject.ToSharedRef(), MetadataNode );
+				}
+
+				bool bAddSuccessful = InternationalizationArchive->AddEntry( AccumulatedNamespace, Source, Translation, MetadataNode, bIsOptional );
+				if( !bAddSuccessful )
+				{
+					UE_LOG( LogInternationalizationArchiveSerializer, Warning,TEXT("Could not add JSON entry to the Internationalization archive: Namespace:%s DefaultText:%s"), *AccumulatedNamespace, *SourceText );
+				}
+			}
+		}
+	}
+
+	if( bConvertSuccess && InJsonObj->HasField( TAG_SUBNAMESPACES ) )
+	{
+		const TArray< TSharedPtr<FJsonValue> > SubnamespaceArray = InJsonObj->GetArrayField( TAG_SUBNAMESPACES );
+
+		for(TArray< TSharedPtr< FJsonValue > >::TConstIterator SubnamespaceIter( SubnamespaceArray.CreateConstIterator() ); SubnamespaceIter; ++SubnamespaceIter )
+		{
+			const TSharedPtr< FJsonValue >  SubnamespaceEntry = *SubnamespaceIter;
+			const TSharedPtr< FJsonObject > SubnamespaceJSONObject = SubnamespaceEntry->AsObject();
+
+			if( !JsonObjToArchive( SubnamespaceJSONObject.ToSharedRef(), AccumulatedNamespace, InternationalizationArchive ) )
+			{
+				bConvertSuccess = false;
+				break;
+			}
+		}
+	}
+
+	return bConvertSuccess;
+}
+
+
+void FJsonInternationalizationArchiveSerializer::GenerateStructuredData( TSharedRef< const FInternationalizationArchive > InInternationalizationArchive, TSharedPtr<FStructuredArchiveEntry> RootElement )
+{
+	//Loop through all the unstructured archive entries and build up our structured hierarchy
+	for(TArchiveEntryContainer::TConstIterator It( InInternationalizationArchive->GetEntryIterator() ); It; ++It)
+	{
+		const TSharedRef< FArchiveEntry > UnstructuredArchiveEntry = It.Value();
+
+		TArray< FString > NamespaceTokens;
+
+		// Tokenize the namespace by using '.' as a delimiter
+		int32 NamespaceTokenCount = UnstructuredArchiveEntry->Namespace.ParseIntoArray( &NamespaceTokens, *NAMESPACE_DELIMITER, true );
+
+		TSharedPtr< FStructuredArchiveEntry > StructuredArchiveEntry = RootElement;
+		//Loop through all the namespace tokens and find the appropriate structured entry, if it does not exist add it.  At the end StructuredArchiveEntry
+		//  will point to the correct hierarchy entry for a given namespace
+		for( int32 TokenIndex = 0; TokenIndex < NamespaceTokenCount; ++TokenIndex )
+		{
+			TSharedPtr<FStructuredArchiveEntry> FoundNamespaceEntry;
+			for( int SubNamespaceIndex = 0; SubNamespaceIndex < StructuredArchiveEntry->SubNamespaces.Num(); SubNamespaceIndex++ )
+			{
+				if(  StructuredArchiveEntry->SubNamespaces[SubNamespaceIndex]->Namespace == NamespaceTokens[TokenIndex] )
+				{
+					FoundNamespaceEntry = StructuredArchiveEntry->SubNamespaces[SubNamespaceIndex];
+					break;
+				}
+			}
+
+			if( !FoundNamespaceEntry.IsValid() )
+			{
+				int32 index = StructuredArchiveEntry->SubNamespaces.Add( MakeShareable( new FStructuredArchiveEntry( NamespaceTokens[TokenIndex] ) ) );
+				FoundNamespaceEntry = StructuredArchiveEntry->SubNamespaces[index];
+			}
+			StructuredArchiveEntry = FoundNamespaceEntry;
+		}
+
+		// We add the unstructured Archive entry to the hierarchy
+		StructuredArchiveEntry->ArchiveEntries.AddUnique( UnstructuredArchiveEntry );
+	}
+}
+
+
+void FJsonInternationalizationArchiveSerializer::SortStructuredData( TSharedPtr< FStructuredArchiveEntry > InElement )
+{
+	if( !InElement.IsValid() )
+	{
+		return;
+	}
+
+	// Sort the manifest entries by source text.
+	InElement->ArchiveEntries.Sort( FCompareArchiveEntryBySourceAndKey() );
+
+	// Sort the subnamespaces by namespace string
+	InElement->SubNamespaces.Sort( FCompareStructuredArchiveEntryByNamespace() );
+
+	// Do the sorting for each of the subnamespaces
+	for( TArray< TSharedPtr< FStructuredArchiveEntry > >::TIterator Iter( InElement->SubNamespaces.CreateIterator() ); Iter; ++Iter )
+	{
+		TSharedPtr< FStructuredArchiveEntry > SubElement = *Iter;
+
+		SortStructuredData( SubElement );
+	}
+}
+
+
+void FJsonInternationalizationArchiveSerializer::StructuredDataToJsonObj( TSharedPtr< const FStructuredArchiveEntry > InElement, TSharedRef< FJsonObject > OutJsonObj )
+{
+	OutJsonObj->SetStringField( TAG_NAMESPACE, InElement->Namespace );
+
+	TArray< TSharedPtr< FJsonValue > > NamespaceArray;
+	TArray< TSharedPtr< FJsonValue > > EntryArray;
+
+	//Write namespace content entries
+	for( TArray< TSharedPtr< FArchiveEntry > >::TConstIterator Iter( InElement->ArchiveEntries ); Iter; ++Iter )
+	{
+		const TSharedPtr< FArchiveEntry > Entry = *Iter;
+		TSharedPtr< FJsonObject > EntryNode = MakeShareable( new FJsonObject );
+
+		FString ProcessedSourceText = Entry->Source.Text;
+		FString ProcessedTranslation = Entry->Translation.Text;
+
+		TSharedPtr< FJsonObject > SourceNode;
+		if( Entry->Source.MetadataObj.IsValid() )
+		{
+			FJsonInternationalizationMetaDataSerializer::SerializeMetadata( Entry->Source.MetadataObj.ToSharedRef(), SourceNode );
+		}
+
+		if( !SourceNode.IsValid() )
+		{
+			SourceNode = MakeShareable( new FJsonObject );
+		}
+
+		SourceNode->SetStringField( TAG_SOURCE_TEXT, ProcessedSourceText );
+		EntryNode->SetObjectField( TAG_SOURCE, SourceNode );
+		
+		TSharedPtr< FJsonObject > TranslationNode;
+		if( Entry->Translation.MetadataObj.IsValid() )
+		{
+			FJsonInternationalizationMetaDataSerializer::SerializeMetadata( Entry->Translation.MetadataObj.ToSharedRef(), TranslationNode );
+		}
+
+		if( !TranslationNode.IsValid() )
+		{
+			TranslationNode = MakeShareable( new FJsonObject );
+		}
+
+		TranslationNode->SetStringField( TAG_TRANSLATION_TEXT, ProcessedTranslation );
+		EntryNode->SetObjectField( TAG_TRANSLATION, TranslationNode );
+
+		if( Entry->KeyMetadataObj.IsValid() )
+		{
+			TSharedPtr< FJsonObject > KeyDataNode;
+			FJsonInternationalizationMetaDataSerializer::SerializeMetadata( Entry->KeyMetadataObj.ToSharedRef(), KeyDataNode );
+			if( KeyDataNode.IsValid() )
+			{
+				EntryNode->SetObjectField( TAG_METADATA_KEY, KeyDataNode );
+			}
+		}
+
+		// We only add the optional field if it is true, it is assumed to be false otherwise.
+		if( Entry->bIsOptional == true )
+		{
+			EntryNode->SetBoolField( TAG_OPTIONAL, Entry->bIsOptional );
+		}
+
+		EntryArray.Add( MakeShareable( new FJsonValueObject( EntryNode ) ) );
+	}
+
+	//Write the subnamespaces
+	for( TArray< TSharedPtr< FStructuredArchiveEntry > >::TConstIterator Iter( InElement->SubNamespaces ); Iter; ++Iter )
+	{
+		const TSharedPtr<FStructuredArchiveEntry> SubElement = *Iter;
+		if( SubElement.IsValid() )
+		{
+			TSharedRef<FJsonObject> SubObject = MakeShareable( new FJsonObject );
+			StructuredDataToJsonObj( SubElement, SubObject );
+
+			NamespaceArray.Add( MakeShareable( new FJsonValueObject( SubObject ) ) );
+		}
+	}
+
+	if( EntryArray.Num() > 0 )
+	{
+		OutJsonObj->SetArrayField( TAG_CHILDREN, EntryArray );
+	}
+
+	if( NamespaceArray.Num() > 0 )
+	{
+		OutJsonObj->SetArrayField( TAG_SUBNAMESPACES, NamespaceArray );
+	}
+}
diff --git a/Engine/Source/Runtime/Internationalization/Private/Serialization/JsonInternationalizationManifestSerializer.cpp b/Engine/Source/Runtime/Internationalization/Private/Serialization/JsonInternationalizationManifestSerializer.cpp
new file mode 100644
--- /dev/null
+++ b/Engine/Source/Runtime/Internationalization/Private/Serialization/JsonInternationalizationManifestSerializer.cpp
@@ -0,0 +1,513 @@
+﻿// Copyright 1998-2014 Epic Games, Inc. All Rights Reserved.
+
+#include "InternationalizationPrivatePCH.h"
+#include "InternationalizationManifest.h"
+#include "Internationalization/InternationalizationMetadata.h"
+#include "Json.h"
+#include "JsonInternationalizationManifestSerializer.h"
+#include "JsonInternationalizationMetadataSerializer.h"
+
+
+DEFINE_LOG_CATEGORY_STATIC(LogInternationalizationManifestSerializer, Log, All);
+
+
+const FString FJsonInternationalizationManifestSerializer::TAG_FORMATVERSION = TEXT("FormatVersion");
+const FString FJsonInternationalizationManifestSerializer::TAG_NAMESPACE = TEXT("Namespace");
+const FString FJsonInternationalizationManifestSerializer::TAG_CHILDREN = TEXT("Children");
+const FString FJsonInternationalizationManifestSerializer::TAG_SUBNAMESPACES = TEXT("Subnamespaces");
+const FString FJsonInternationalizationManifestSerializer::TAG_PATH = TEXT("Path");
+const FString FJsonInternationalizationManifestSerializer::TAG_OPTIONAL = TEXT("Optional");
+const FString FJsonInternationalizationManifestSerializer::TAG_KEYCOLLECTION = TEXT("Keys");
+const FString FJsonInternationalizationManifestSerializer::TAG_KEY = TEXT("Key");
+const FString FJsonInternationalizationManifestSerializer::TAG_DEPRECATED_DEFAULTTEXT = TEXT("DefaultText");
+const FString FJsonInternationalizationManifestSerializer::TAG_SOURCE = TEXT("Source");
+const FString FJsonInternationalizationManifestSerializer::TAG_SOURCE_TEXT = TEXT("Text");
+const FString FJsonInternationalizationManifestSerializer::TAG_METADATA = TEXT("MetaData");
+const FString FJsonInternationalizationManifestSerializer::TAG_METADATA_INFO = TEXT("Info");
+const FString FJsonInternationalizationManifestSerializer::TAG_METADATA_KEY = TEXT("Key");
+const FString FJsonInternationalizationManifestSerializer::NAMESPACE_DELIMITER = TEXT(".");
+
+
+struct FCompareManifestEntryBySource
+{
+	FORCEINLINE bool operator()( TSharedPtr< FManifestEntry > A, TSharedPtr< FManifestEntry > B ) const
+	{
+		return A->Source < B->Source;
+	}
+};
+
+
+struct FCompareStructuredEntryByNamespace
+{
+	FORCEINLINE bool operator()( TSharedPtr< FStructuredEntry > A, TSharedPtr< FStructuredEntry > B ) const
+	{
+		return A->Namespace < B->Namespace;
+	}
+};
+
+
+bool FJsonInternationalizationManifestSerializer::DeserializeManifest( FArchive& Archive, TSharedRef< FInternationalizationManifest > Manifest )
+{
+	TSharedPtr< FJsonObject > JsonManifestObj;
+	TSharedRef< TJsonReader<> > Reader = TJsonReaderFactory<>::Create( &Archive );
+	bool bExecSuccessful = FJsonSerializer::Deserialize( Reader, JsonManifestObj );
+	
+	if ( bExecSuccessful && JsonManifestObj.IsValid() )
+	{
+		bExecSuccessful = DeserializeInternal( JsonManifestObj.ToSharedRef(), Manifest );
+	}		
+
+	return bExecSuccessful;
+}
+
+
+bool FJsonInternationalizationManifestSerializer::DeserializeManifest( const FString& InStr, TSharedRef< FInternationalizationManifest > Manifest )
+{
+	TSharedPtr<FJsonObject> JsonManifestObj;
+	TSharedRef< TJsonReader<> > Reader = TJsonReaderFactory<>::Create( InStr );
+	bool bExecSuccessful = FJsonSerializer::Deserialize( Reader, JsonManifestObj );
+
+	if ( bExecSuccessful && JsonManifestObj.IsValid() )
+	{
+		bExecSuccessful = DeserializeInternal( JsonManifestObj.ToSharedRef(), Manifest );
+	}
+
+	return bExecSuccessful;
+}
+
+
+bool FJsonInternationalizationManifestSerializer::DeserializeManifest( TSharedRef< FJsonObject > InJsonObj, TSharedRef< FInternationalizationManifest > Manifest )
+{
+	return DeserializeInternal( InJsonObj, Manifest );
+}
+
+
+bool FJsonInternationalizationManifestSerializer::SerializeManifest( TSharedRef< const FInternationalizationManifest > Manifest, FArchive& Archive )
+{
+	TSharedRef< FJsonObject > JsonManifestObj = MakeShareable( new FJsonObject );
+	bool bExecSuccessful = SerializeInternal( Manifest, JsonManifestObj );
+
+	if( bExecSuccessful )
+	{
+		TSharedRef< TJsonWriter<> > Writer = TJsonWriterFactory<>::Create( &Archive );
+		bExecSuccessful = FJsonSerializer::Serialize( JsonManifestObj, Writer );
+		Writer->Close();
+	}
+	return bExecSuccessful;
+}
+
+
+bool FJsonInternationalizationManifestSerializer::SerializeManifest( TSharedRef< const FInternationalizationManifest > Manifest, FString& Str )
+{
+	TSharedRef< FJsonObject > JsonManifestObj = MakeShareable( new FJsonObject );
+	bool bExecSuccessful = SerializeInternal( Manifest, JsonManifestObj );
+
+	if(bExecSuccessful)
+	{
+		TSharedRef< TJsonWriter<> > Writer = TJsonWriterFactory<>::Create( &Str );
+		bExecSuccessful = FJsonSerializer::Serialize(JsonManifestObj, Writer);
+		Writer->Close();
+	}
+	return bExecSuccessful;
+}
+
+
+bool FJsonInternationalizationManifestSerializer::SerializeManifest( TSharedRef< const FInternationalizationManifest > Manifest, TSharedRef< FJsonObject > JsonObj )
+{
+	return SerializeInternal( Manifest, JsonObj );
+}
+
+
+bool FJsonInternationalizationManifestSerializer::DeserializeInternal( TSharedRef< FJsonObject > InJsonObj, TSharedRef< FInternationalizationManifest > Manifest )
+{
+	if( InJsonObj->HasField( TAG_FORMATVERSION ) )
+	{
+		const int32 FormatVersion = static_cast<int32>(InJsonObj->GetNumberField( TAG_FORMATVERSION ));
+		Manifest->SetFormatVersion(static_cast<FInternationalizationManifest::EFormatVersion>(FormatVersion));
+	}
+	else
+	{
+		Manifest->SetFormatVersion(FInternationalizationManifest::EFormatVersion::Initial);
+	}
+
+	return JsonObjToManifest(InJsonObj, TEXT(""), Manifest );
+}
+
+
+bool FJsonInternationalizationManifestSerializer::SerializeInternal( TSharedRef< const FInternationalizationManifest > InManifest, TSharedRef< FJsonObject > JsonObj )
+{
+	TSharedPtr< FStructuredEntry > RootElement = MakeShareable( new FStructuredEntry( TEXT("") ) );
+
+	// Condition the data so that it exists in a structured hierarchy for easy population of the JSON object.
+	GenerateStructuredData( InManifest, RootElement );
+
+	// Arrange the entries in non-cultural format so that diffs are easier to read.
+	SortStructuredData( RootElement );
+
+	//Clear out anything that may be in the JSON object
+	JsonObj->Values.Empty();
+
+	// Set format version.
+	JsonObj->SetNumberField(TAG_FORMATVERSION, static_cast<double>(InManifest->GetFormatVersion()));
+
+	// Setup the JSON object using the structured data created
+	StructuredDataToJsonObj( RootElement, JsonObj );
+
+	return true;
+}
+
+
+bool FJsonInternationalizationManifestSerializer::JsonObjToManifest( TSharedRef< FJsonObject > InJsonObj, FString ParentNamespace, TSharedRef< FInternationalizationManifest > Manifest )
+{
+	bool bConvertSuccess = true;
+	FString AccumulatedNamespace = ParentNamespace;
+	if( InJsonObj->HasField( TAG_NAMESPACE) )
+	{
+		if( !( AccumulatedNamespace.IsEmpty() ) )
+		{
+			AccumulatedNamespace += NAMESPACE_DELIMITER;
+		}
+		AccumulatedNamespace += InJsonObj->GetStringField( TAG_NAMESPACE );
+	}
+	else
+	{
+		// We found an entry with a missing namespace
+		bConvertSuccess = false;
+	}
+
+	// Process all the child objects
+	if( bConvertSuccess && InJsonObj->HasField( TAG_CHILDREN ) )
+	{
+		const TArray< TSharedPtr< FJsonValue> > ChildrenArray = InJsonObj->GetArrayField( TAG_CHILDREN );
+
+		for(TArray< TSharedPtr< FJsonValue > >::TConstIterator ChildIter( ChildrenArray.CreateConstIterator() ); ChildIter && bConvertSuccess; ++ChildIter)
+		{
+			const TSharedPtr< FJsonValue >  ChildEntry = *ChildIter;
+			const TSharedPtr< FJsonObject > ChildJSONObject = ChildEntry->AsObject();
+
+			
+			FString SourceText;
+			TSharedPtr< FLocMetadataObject > SourceMetadata;
+			if( ChildJSONObject->HasTypedField< EJson::String>( TAG_DEPRECATED_DEFAULTTEXT) )
+			{
+				SourceText = ChildJSONObject->GetStringField( TAG_DEPRECATED_DEFAULTTEXT );
+			}
+			else if( ChildJSONObject->HasTypedField< EJson::Object>( TAG_SOURCE ) )
+			{
+				const TSharedPtr< FJsonObject > SourceJSONObject = ChildJSONObject->GetObjectField( TAG_SOURCE );
+				if( SourceJSONObject->HasTypedField< EJson::String >( TAG_SOURCE_TEXT ) )
+				{
+					SourceText = SourceJSONObject->GetStringField( TAG_SOURCE_TEXT );
+
+					// Source meta data is mixed in with the source text, we'll process metadata if the source json object has more than one entry
+					if( SourceJSONObject->Values.Num() > 1 )
+					{
+						// We load in the entire source object as metadata and remove the source object
+						FJsonInternationalizationMetaDataSerializer::DeserializeMetadata( SourceJSONObject.ToSharedRef(), SourceMetadata );
+						if( SourceMetadata.IsValid() )
+						{
+							SourceMetadata->Values.Remove( TAG_SOURCE_TEXT );
+						}
+					}
+				}
+				else
+				{
+					bConvertSuccess = false;
+				}
+			}
+			else
+			{
+				bConvertSuccess = false;
+			}
+
+			FLocItem Source(SourceText);
+			Source.MetadataObj = SourceMetadata;
+
+			if( bConvertSuccess && ChildJSONObject->HasField( TAG_KEYCOLLECTION ) )
+			{
+				
+				const TArray< TSharedPtr<FJsonValue> > ContextArray = ChildJSONObject->GetArrayField( TAG_KEYCOLLECTION);
+
+				for(TArray< TSharedPtr< FJsonValue > >::TConstIterator ContextIter( ContextArray.CreateConstIterator() ); ContextIter && bConvertSuccess; ++ContextIter)
+				{
+					const TSharedPtr< FJsonValue > ContextEntry = *ContextIter;
+					const TSharedPtr< FJsonObject > ContextJSONObject = ContextEntry->AsObject();
+
+					if( ContextJSONObject->HasTypedField< EJson::String >( TAG_KEY ) )
+					{
+						const FString Key = ContextJSONObject->GetStringField( TAG_KEY );
+						const FString SourceLocation = ContextJSONObject->HasField( TAG_PATH ) ? ContextJSONObject->GetStringField( TAG_PATH ) : FString();
+
+						FContext CommandContext;
+						CommandContext.Key = Key;
+						CommandContext.SourceLocation = SourceLocation;
+
+						if( ContextJSONObject->HasTypedField< EJson::Boolean >( TAG_OPTIONAL ) )
+						{
+							CommandContext.bIsOptional = ContextJSONObject->GetBoolField( TAG_OPTIONAL );
+						}
+
+						if( ContextJSONObject->HasTypedField< EJson::Object >( TAG_METADATA ) )
+						{
+							const TSharedPtr< FJsonObject > MetaDataJSONObject = ContextJSONObject->GetObjectField( TAG_METADATA );
+
+							if( MetaDataJSONObject->HasTypedField< EJson::Object >( TAG_METADATA_INFO ) )
+							{
+								const TSharedPtr< FJsonObject > MetaDataInfoJSONObject = MetaDataJSONObject->GetObjectField( TAG_METADATA_INFO );
+
+								TSharedPtr< FLocMetadataObject > MetadataNode;
+								FJsonInternationalizationMetaDataSerializer::DeserializeMetadata( MetaDataInfoJSONObject.ToSharedRef(), MetadataNode );
+								if( MetadataNode.IsValid() )
+								{
+									CommandContext.InfoMetadataObj = MetadataNode;
+								}
+							}
+
+							if( MetaDataJSONObject->HasTypedField< EJson::Object >( TAG_METADATA_KEY ) )
+							{
+								const TSharedPtr< FJsonObject > MetaDataKeyJSONObject = MetaDataJSONObject->GetObjectField( TAG_METADATA_KEY );
+
+								TSharedPtr< FLocMetadataObject > MetadataNode;
+								FJsonInternationalizationMetaDataSerializer::DeserializeMetadata( MetaDataKeyJSONObject.ToSharedRef(), MetadataNode );
+								if( MetadataNode.IsValid() )
+								{
+									CommandContext.KeyMetadataObj = MetadataNode;
+								}
+							}
+						}
+						bool bAddSuccessful = Manifest->AddSource( AccumulatedNamespace, Source, CommandContext );
+						if(!bAddSuccessful)
+						{
+							UE_LOG( LogInternationalizationManifestSerializer, Warning,TEXT("Could not add JSON entry to the Internationalization manifest: Namespace:%s SourceText:%s SourceData:%s"), 
+								*AccumulatedNamespace, 
+								*Source.Text, 
+								*FJsonInternationalizationMetaDataSerializer::MetadataToString(Source.MetadataObj) );
+						}
+					}
+					else
+					{
+						//We found a context entry that is missing a identifier/key or a path
+						bConvertSuccess = false;
+						break;
+					}
+
+				}
+			}
+			else
+			{
+				// We have an entry that is missing a key/context collection or default text entry.
+				bConvertSuccess = false;
+				break;
+			}
+
+		}
+	}
+
+	if( bConvertSuccess && InJsonObj->HasField( TAG_SUBNAMESPACES ) )
+	{
+		const TArray< TSharedPtr<FJsonValue> > SubnamespaceArray = InJsonObj->GetArrayField( TAG_SUBNAMESPACES );
+
+		for(TArray< TSharedPtr< FJsonValue > >::TConstIterator SubnamespaceIter( SubnamespaceArray.CreateConstIterator() ); SubnamespaceIter; ++SubnamespaceIter )
+		{
+			const TSharedPtr< FJsonValue >  SubnamespaceEntry = *SubnamespaceIter;
+			const TSharedPtr< FJsonObject > SubnamespaceJSONObject = SubnamespaceEntry->AsObject();
+
+			if( !JsonObjToManifest( SubnamespaceJSONObject.ToSharedRef(), AccumulatedNamespace, Manifest ) )
+			{
+				bConvertSuccess = false;
+				break;
+			}
+		}
+	}
+	
+	return bConvertSuccess;
+}
+
+
+void FJsonInternationalizationManifestSerializer::GenerateStructuredData( TSharedRef< const FInternationalizationManifest > InManifest, TSharedPtr< FStructuredEntry > RootElement )
+{
+	//Loop through all the unstructured manifest entries and build up our structured hierarchy
+	for( TManifestEntryByContextIdContainer::TConstIterator It( InManifest->GetEntriesByContextIdIterator() ); It; ++It )
+	{
+		const TSharedRef< FManifestEntry > UnstructuredManifestEntry = It.Value();
+
+		TArray< FString > NamespaceTokens;
+
+		// Tokenize the namespace by using '.' as a delimiter
+		int32 NamespaceTokenCount = UnstructuredManifestEntry->Namespace.ParseIntoArray( &NamespaceTokens, *NAMESPACE_DELIMITER, true );
+
+		TSharedPtr< FStructuredEntry > StructuredManifestEntry = RootElement;
+		//Loop through all the namespace tokens and find the appropriate structured entry, if it does not exist add it.  At the end StructuredManifestEntry
+		//  will point to the correct hierarchy entry for a given namespace
+		for( int32 TokenIndex = 0; TokenIndex < NamespaceTokenCount; ++TokenIndex )
+		{
+			TSharedPtr< FStructuredEntry > FoundNamespaceEntry;
+			for( int SubNamespaceIndex = 0; SubNamespaceIndex < StructuredManifestEntry->SubNamespaces.Num(); SubNamespaceIndex++ )
+			{
+				if(  StructuredManifestEntry->SubNamespaces[SubNamespaceIndex]->Namespace == NamespaceTokens[TokenIndex] )
+				{
+					FoundNamespaceEntry = StructuredManifestEntry->SubNamespaces[SubNamespaceIndex];
+					break;
+				}
+			}
+
+			if( !FoundNamespaceEntry.IsValid() )
+			{
+				int32 index = StructuredManifestEntry->SubNamespaces.Add( MakeShareable( new FStructuredEntry( NamespaceTokens[TokenIndex] ) ) );
+				FoundNamespaceEntry = StructuredManifestEntry->SubNamespaces[index];
+			}
+			StructuredManifestEntry = FoundNamespaceEntry;
+		}
+
+		// We add the unstructured manifest entry to the hierarchy
+		StructuredManifestEntry->ManifestEntries.AddUnique( UnstructuredManifestEntry );
+	}
+}
+
+
+void FJsonInternationalizationManifestSerializer::SortStructuredData( TSharedPtr< FStructuredEntry > InElement )
+{
+	if( !InElement.IsValid() )
+	{
+		return;
+	}
+
+	// Sort the manifest entries by source text.
+	InElement->ManifestEntries.Sort( FCompareManifestEntryBySource() );
+
+	// Sort the manifest entry contexts by key/identifier
+	for( TArray< TSharedPtr< FManifestEntry > >::TIterator Iter( InElement->ManifestEntries.CreateIterator() ); Iter; ++Iter)
+	{
+		TSharedPtr< FManifestEntry > SubEntry = *Iter;
+		if( SubEntry.IsValid())
+		{
+			SubEntry->Contexts.Sort( );
+		}
+	}
+	
+	// Sort the subnamespaces by namespace string
+	InElement->SubNamespaces.Sort( FCompareStructuredEntryByNamespace() );
+
+	// Do the sorting for each of the subnamespaces
+	for( TArray< TSharedPtr< FStructuredEntry > >::TIterator Iter( InElement->SubNamespaces.CreateIterator() ); Iter; ++Iter )
+	{
+		TSharedPtr< FStructuredEntry > SubElement = *Iter;
+
+		SortStructuredData( SubElement );
+	}
+}
+
+
+void FJsonInternationalizationManifestSerializer::StructuredDataToJsonObj( TSharedPtr< const FStructuredEntry > InElement, TSharedRef< FJsonObject > JsonObj )
+{
+
+	JsonObj->SetStringField( TAG_NAMESPACE, InElement->Namespace );
+
+	TArray< TSharedPtr< FJsonValue > > NamespaceArray;
+	TArray< TSharedPtr< FJsonValue > > EntryArray;
+
+	//Write namespace content entries
+	for( TArray< TSharedPtr< FManifestEntry > >::TConstIterator Iter( InElement->ManifestEntries ); Iter; ++Iter )
+	{
+		const TSharedPtr< FManifestEntry > Entry = *Iter;
+		TSharedPtr< FJsonObject > EntryNode = MakeShareable( new FJsonObject );
+
+		TSharedPtr<FJsonObject> SourceNode;
+
+		if( Entry->Source.MetadataObj.IsValid() )
+		{
+			FJsonInternationalizationMetaDataSerializer::SerializeMetadata( Entry->Source.MetadataObj.ToSharedRef(), SourceNode );
+		}
+		
+		if( !SourceNode.IsValid() )
+		{
+			SourceNode = MakeShareable( new FJsonObject );
+		}
+
+		// Add escapes for special chars - doesn't do backslash
+		{
+			FString ProcessedText = Entry->Source.Text;
+			SourceNode->SetStringField(TAG_SOURCE_TEXT, ProcessedText);
+		}
+
+		EntryNode->SetObjectField( TAG_SOURCE, SourceNode );
+
+		TArray< TSharedPtr< FJsonValue > > KeyArray;
+
+		for(auto ContextIter = Entry->Contexts.CreateConstIterator(); ContextIter; ++ContextIter)
+		{
+			const FContext& AContext = *ContextIter;
+
+			FString ProcessedText = AContext.SourceLocation;
+			ProcessedText.ReplaceInline( TEXT("\\"), TEXT("/"));
+			ProcessedText.ReplaceInline( *FPaths::RootDir(), TEXT("/"));
+
+			TSharedPtr<FJsonObject> KeyNode = MakeShareable( new FJsonObject );
+			KeyNode->SetStringField( TAG_KEY, AContext.Key );
+			KeyNode->SetStringField( TAG_PATH, ProcessedText );
+
+			// We only add the optional field if it is true, it is assumed to be false otherwise.
+			if( AContext.bIsOptional == true )
+			{
+				KeyNode->SetBoolField( TAG_OPTIONAL, AContext.bIsOptional);
+			}
+
+			TSharedPtr< FJsonObject > MetaDataNode = MakeShareable( new FJsonObject );
+
+			if( AContext.InfoMetadataObj.IsValid() )
+			{
+				TSharedPtr< FJsonObject > InfoDataNode;
+				FJsonInternationalizationMetaDataSerializer::SerializeMetadata( AContext.InfoMetadataObj.ToSharedRef(), InfoDataNode );
+				if( InfoDataNode.IsValid() )
+				{
+					MetaDataNode->SetObjectField( TAG_METADATA_INFO, InfoDataNode );
+				}
+			}
+
+			if( AContext.KeyMetadataObj.IsValid() )
+			{
+				TSharedPtr< FJsonObject > KeyDataNode;
+				FJsonInternationalizationMetaDataSerializer::SerializeMetadata( AContext.KeyMetadataObj.ToSharedRef(), KeyDataNode );
+				if( KeyDataNode.IsValid() )
+				{
+					MetaDataNode->SetObjectField( TAG_METADATA_KEY, KeyDataNode );
+				}
+			}
+
+			if( MetaDataNode->Values.Num() > 0 )
+			{
+				KeyNode->SetObjectField( TAG_METADATA, MetaDataNode );
+			}
+			
+			KeyArray.Add( MakeShareable( new FJsonValueObject( KeyNode ) ) );
+		}
+
+		EntryNode->SetArrayField( TAG_KEYCOLLECTION, KeyArray);
+
+		EntryArray.Add( MakeShareable( new FJsonValueObject( EntryNode ) ) );
+	}
+
+	//Write the subnamespaces
+	for( TArray< TSharedPtr< FStructuredEntry > >::TConstIterator Iter( InElement->SubNamespaces.CreateConstIterator() ); Iter; ++Iter )
+	{
+		const TSharedPtr< FStructuredEntry > SubElement = *Iter;
+
+		if(SubElement.IsValid())
+		{
+			TSharedRef< FJsonObject > SubObject = MakeShareable( new FJsonObject );
+			StructuredDataToJsonObj( SubElement, SubObject );
+
+			NamespaceArray.Add( MakeShareable( new FJsonValueObject( SubObject ) ) );
+		}
+	}
+
+	if( EntryArray.Num() > 0 )
+	{
+		JsonObj->SetArrayField( TAG_CHILDREN, EntryArray );
+	}
+
+	if( NamespaceArray.Num() > 0 )
+	{
+		JsonObj->SetArrayField( TAG_SUBNAMESPACES, NamespaceArray );
+	}
+}
diff --git a/Engine/Source/Runtime/Internationalization/Private/Serialization/JsonInternationalizationMetadataSerializer.cpp b/Engine/Source/Runtime/Internationalization/Private/Serialization/JsonInternationalizationMetadataSerializer.cpp
new file mode 100644
--- /dev/null
+++ b/Engine/Source/Runtime/Internationalization/Private/Serialization/JsonInternationalizationMetadataSerializer.cpp
@@ -0,0 +1,200 @@
+// Copyright 1998-2014 Epic Games, Inc. All Rights Reserved.
+
+#include "InternationalizationPrivatePCH.h"
+#include "Internationalization/InternationalizationMetadata.h"
+#include "Json.h"
+
+
+struct FCompareLocMetadataValue
+{
+	FORCEINLINE bool operator()( TSharedPtr< FLocMetadataValue > A, TSharedPtr< FLocMetadataValue > B ) const
+	{
+		if( A.IsValid() && B.IsValid() )
+		{
+			return (*A < *B);
+		}
+		else if( A.IsValid() != B.IsValid() )
+		{
+			return B.IsValid();
+		}
+		return false;
+	}
+};
+
+
+TSharedPtr<FLocMetadataValue> JSonValueToLocMetaDataValue( const TSharedRef< FJsonValue > JsonValue )
+{
+	switch( JsonValue->Type )
+	{
+	case EJson::Boolean:
+		{
+			return MakeShareable( new FLocMetadataValueBoolean( JsonValue->AsBool() ) );
+		}
+		break;
+	case EJson::String:
+		{
+			return MakeShareable( new FLocMetadataValueString( JsonValue->AsString() ) );
+		}
+		break;
+	case EJson::Array:
+		{
+			TArray< TSharedPtr< FLocMetadataValue > > MetadataArray;
+			TArray< TSharedPtr< FJsonValue > > JsonArray = JsonValue->AsArray();
+			for( auto ArrayIter = JsonArray.CreateConstIterator(); ArrayIter; ++ArrayIter )
+			{
+				const TSharedPtr< FJsonValue >& Item = *ArrayIter;
+				if( Item.IsValid() )
+				{
+					TSharedPtr< FLocMetadataValue > Entry = JSonValueToLocMetaDataValue( Item.ToSharedRef() );
+					if( Entry.IsValid() )
+					{
+						MetadataArray.Add( Entry );
+					}
+				}
+			}
+			if( MetadataArray.Num() > 0 )
+			{
+				return MakeShareable( new FLocMetadataValueArray( MetadataArray ) );
+			}
+		}
+		break;
+	case EJson::Object:
+		{
+			TSharedPtr< FLocMetadataObject > MetadataSubObject = MakeShareable( new FLocMetadataObject );
+			TSharedPtr< FJsonObject > JsonObject = JsonValue->AsObject();
+			for( auto ValueIter = JsonObject->Values.CreateConstIterator(); ValueIter; ++ValueIter )
+			{
+				const FString Name = (*ValueIter).Key;
+				TSharedPtr< FJsonValue > Value = (*ValueIter).Value;
+				if( Value.IsValid() )
+				{
+					TSharedPtr< FLocMetadataValue > MetadataValue = JSonValueToLocMetaDataValue( Value.ToSharedRef() );
+					if( MetadataValue.IsValid() )
+					{
+						MetadataSubObject->SetField( Name, MetadataValue );
+					}
+				}
+			}
+			return MakeShareable( new FLocMetadataValueObject(MetadataSubObject) );
+		}
+		break;
+	default:
+		{
+			// At the moment we do not support all the json types as metadata.  In the future these types will be handled in a way that they can be stored in an unprocessed way.
+		}
+		break;
+	}
+
+	return NULL;
+}
+
+
+TSharedPtr<FJsonValue> LocMetaDataValueToJsonValue( const TSharedRef< FLocMetadataValue > MetadataValue )
+{
+	switch( MetadataValue->Type )
+	{
+	case ELocMetadataType::Boolean:
+		{
+			return MakeShareable( new FJsonValueBoolean( MetadataValue->AsBool() ) );
+		}
+		break;
+
+	case ELocMetadataType::String:
+		{
+			return MakeShareable( new FJsonValueString( MetadataValue->AsString() ) );
+		}
+		break;
+
+	case ELocMetadataType::Array:
+		{
+			TArray< TSharedPtr<FJsonValue> > JsonArrayVals;
+			TArray< TSharedPtr< FLocMetadataValue > > MetaDataArray = MetadataValue->AsArray();
+			
+			MetaDataArray.Sort( FCompareLocMetadataValue() );
+			for( auto ArrayIter = MetaDataArray.CreateConstIterator(); ArrayIter; ++ArrayIter )
+			{
+				const TSharedPtr<FLocMetadataValue>& Item = *ArrayIter;
+				if( Item.IsValid() )
+				{
+					TSharedPtr<FJsonValue> Entry = LocMetaDataValueToJsonValue( Item.ToSharedRef() );
+					if( Entry.IsValid() )
+					{
+						JsonArrayVals.Add( Entry );
+					}
+				}
+			}
+
+			if( JsonArrayVals.Num() > 0 )
+			{
+				return MakeShareable( new FJsonValueArray( JsonArrayVals ) );
+			}
+		}
+		break;
+
+	case ELocMetadataType::Object:
+		{
+			TSharedPtr< FJsonObject > JsonSubObject = MakeShareable( new FJsonObject );
+			TSharedPtr< FLocMetadataObject > MetaDataObject = MetadataValue->AsObject();
+			for( auto ValueIter = MetaDataObject->Values.CreateConstIterator(); ValueIter; ++ValueIter )
+			{
+				const FString Name = (*ValueIter).Key;
+				TSharedPtr< FLocMetadataValue > Value = (*ValueIter).Value;
+				if( Value.IsValid() )
+				{
+					TSharedPtr< FJsonValue > JsonValue = LocMetaDataValueToJsonValue( Value.ToSharedRef() );
+					if( JsonValue.IsValid() )
+					{
+						JsonSubObject->SetField( Name, JsonValue );
+					}
+				}
+			}
+
+			// Sorting by key is probably sufficient for now but ideally we would sort the resulting json object using
+			//  the same logic that is in the FLocMetadata < operator
+			JsonSubObject->Values.KeySort( TLess<FString>() );
+
+			return MakeShareable( new FJsonValueObject(JsonSubObject) );
+		}
+
+	default:
+		break;
+	}
+	return NULL;
+}
+
+
+void FJsonInternationalizationMetaDataSerializer::DeserializeMetadata( const TSharedRef< FJsonObject > JsonObj, TSharedPtr< FLocMetadataObject >& OutMetaDataObj )
+{
+	OutMetaDataObj = JSonValueToLocMetaDataValue( MakeShareable( new FJsonValueObject( JsonObj ) ) )->AsObject();
+}
+
+
+void FJsonInternationalizationMetaDataSerializer::SerializeMetadata( const TSharedRef< FLocMetadataObject > MetaData, TSharedPtr< FJsonObject >& OutJsonObj )
+{
+	OutJsonObj = LocMetaDataValueToJsonValue( MakeShareable( new FLocMetadataValueObject( MetaData ) ) )->AsObject();
+}
+
+
+FString FJsonInternationalizationMetaDataSerializer::MetadataToString( const TSharedPtr<FLocMetadataObject> Metadata )
+{
+	FString StringMetadata = TEXT("");
+	if( Metadata.IsValid() )
+	{
+		TSharedPtr< FJsonObject > JsonKeyMetadata;
+		FJsonInternationalizationMetaDataSerializer::SerializeMetadata( Metadata.ToSharedRef(), JsonKeyMetadata );
+
+		if( JsonKeyMetadata.IsValid() )
+		{
+			JsonKeyMetadata->Values.KeySort( TLess<FString>() );
+
+			TSharedRef< TJsonWriter<> > Writer = TJsonWriterFactory<>::Create( &StringMetadata );
+			FJsonSerializer::Serialize( JsonKeyMetadata.ToSharedRef(), Writer );
+			Writer->Close();
+
+			StringMetadata.ReplaceInline( TEXT("\t"), TEXT(" ") );
+			StringMetadata.ReplaceInline( TEXT("\r\n"), TEXT(" ") );
+			StringMetadata.ReplaceInline( TEXT("\n"), TEXT(" ") );
+		}
+	}
+	return StringMetadata;
+}
diff --git a/Engine/Source/Runtime/Internationalization/Public/Internationalization.h b/Engine/Source/Runtime/Internationalization/Public/Internationalization.h
new file mode 100644
--- /dev/null
+++ b/Engine/Source/Runtime/Internationalization/Public/Internationalization.h
@@ -0,0 +1,14 @@
+// Copyright 1998-2014 Epic Games, Inc. All Rights Reserved.
+
+#pragma once
+
+
+/* Public dependencies
+ *****************************************************************************/
+
+#include "Core.h"
+
+
+/* Public includes
+ *****************************************************************************/
+
diff --git a/Engine/Source/Runtime/Internationalization/Public/Serialization/JsonInternationalizationArchiveSerializer.h b/Engine/Source/Runtime/Internationalization/Public/Serialization/JsonInternationalizationArchiveSerializer.h
new file mode 100644
--- /dev/null
+++ b/Engine/Source/Runtime/Internationalization/Public/Serialization/JsonInternationalizationArchiveSerializer.h
@@ -0,0 +1,167 @@
+// Copyright 1998-2014 Epic Games, Inc. All Rights Reserved.
+
+#pragma once
+
+#include "IInternationalizationArchiveSerializer.h"
+#include "Json.h"
+
+
+/**
+ * Used to arrange Internationalization archive data in a hierarchy based on namespace prior to json serialization.
+ */
+struct FStructuredArchiveEntry
+{
+public:
+	FStructuredArchiveEntry( const FString& InNamespace )
+		: Namespace( InNamespace )
+	{ }
+
+	const FString Namespace;
+	TArray< TSharedPtr< FStructuredArchiveEntry > > SubNamespaces;
+	TArray< TSharedPtr< class FArchiveEntry > > ArchiveEntries;
+};
+
+
+/**
+ * Implements a InternationalizationArchive serializer that serializes to and from Json encoded data.
+ */
+class INTERNATIONALIZATION_API FJsonInternationalizationArchiveSerializer
+	: public IInternationalizationArchiveSerializer
+{
+public:
+
+	/** Default constructor. */
+	FJsonInternationalizationArchiveSerializer() { }
+
+	/**
+	 * Deserializes a Internationalization archive from archive.
+	 *
+	 * @param Archive The archive to serialize from.
+	 * @param InternationalizationArchive The populated Internationalization archive.
+	 *
+	 * @return true if deserialization was successful, false otherwise.
+	 */
+	virtual bool DeserializeArchive( FArchive& Archive, TSharedRef< FInternationalizationArchive > InternationalizationArchive ) override;
+
+	/**
+	 * Deserializes a Internationalization archive from a string.
+	 *
+	 * @param InStr The string to serialize from.
+	 * @param InternationalizationArchive The populated Internationalization archive.
+	 *
+	 * @return true if deserialization was successful, false otherwise.
+	 */
+	virtual bool DeserializeArchive( const FString& InStr, TSharedRef< FInternationalizationArchive > InternationalizationArchive );
+
+	/**
+	 * Deserializes a Internationalization archive from a JSON object.
+	 *
+	 * @param InJsonObj The JSON object to serialize from.
+	 * @param InternationalizationArchive The populated Internationalization archive.
+	 *
+	 * @return true if deserialization was successful, false otherwise.
+	 */
+	virtual bool DeserializeArchive( TSharedRef< FJsonObject > InJsonObj, TSharedRef< FInternationalizationArchive > InternationalizationArchive );
+
+	/**
+	 * Serializes a Internationalization archive into an archive.
+	 *
+	 * @param InternationalizationArchive The Internationalization archive data to serialize.
+	 * @param Archive The archive to serialize into.
+	 *
+	 * @return true if serialization was successful, false otherwise.
+	 */
+	virtual bool SerializeArchive( TSharedRef< const FInternationalizationArchive > InternationalizationArchive, FArchive& Archive ) override;
+
+	/**
+	 * Serializes a Internationalization archive to a string.
+	 *
+	 * @param InternationalizationArchive The Internationalization archive data to serialize.
+	 * @param Str The string to serialize into.
+	 *
+	 * @return true if serialization was successful, false otherwise.
+	 */
+	virtual bool SerializeArchive( TSharedRef< const FInternationalizationArchive > InternationalizationArchive, FString& Str );
+
+	/**
+	 * Serializes a Internationalization archive to a JSON object.
+	 *
+	 * @param InternationalizationArchive The Internationalization archive data to serialize.
+	 * @param JsonObj The JSON object to serialize into.
+	 *
+	 * @return true if serialization was successful, false otherwise.
+	 */
+	virtual bool SerializeArchive( TSharedRef< const FInternationalizationArchive > InternationalizationArchive, TSharedRef< FJsonObject > JsonObj );
+
+protected:
+
+	/**
+	 * Convert a JSON object to a Internationalization archive.
+	 *
+	 * @param InJsonObj The JSON object to serialize from.
+	 * @param InternationalizationArchive The Internationalization archive that will store the data.
+	 *
+	 * @return true if deserialization was successful, false otherwise.
+	 */
+	bool DeserializeInternal( TSharedRef< FJsonObject > InJsonObj, TSharedRef< FInternationalizationArchive > InternationalizationArchive );
+
+	/**
+	 * Convert a Internationalization archive to a JSON object.
+	 *
+	 * @param InInternationalizationArchive The Internationalization archive object to serialize from.
+	 * @param JsonObj The Json object that will store the data.
+	 *
+	 * @return true if serialization was successful, false otherwise.
+	 */
+	bool SerializeInternal( TSharedRef< const FInternationalizationArchive > InInternationalizationArchive, TSharedRef< FJsonObject > JsonObj );
+
+	/**
+	 * Recursive function that will traverse the JSON object and populate a Internationalization archive.
+	 *
+	 * @param InJsonObj The JSON object.
+	 * @param ParentNamespace The namespace of the parent JSON object.
+	 * @param InternationalizationArchive The Internationalization archive that will store the data.
+	 *
+	 * @return true if successful, false otherwise.
+	 */
+	bool JsonObjToArchive( TSharedRef< FJsonObject > InJsonObj, FString ParentNamespace, TSharedRef< FInternationalizationArchive > InternationalizationArchive );
+
+	/**
+	 * Takes a Internationalization archive and arranges the data into a hierarchy based on namespace.
+	 *
+	 * @param InInternationalizationArchive The Internationalization archive.
+	 * @param RootElement The root element of the structured data.
+	 */
+	void GenerateStructuredData( TSharedRef< const FInternationalizationArchive > InInternationalizationArchive, TSharedPtr< FStructuredArchiveEntry > RootElement );
+
+	/**
+	 * Goes through the structured, hierarchy based, archive data and does a non-culture specific sort on namespaces and default text.
+	 *
+	 * @param RootElement The root element of the structured data.
+	 */
+	void SortStructuredData( TSharedPtr< FStructuredArchiveEntry > InElement );
+
+	/**
+	 * Populates a JSON object from Internationalization archive data that has been structured based on namespace.
+	 *
+	 * @param InElement Internationalization archive data structured based on namespace.
+	 * @param JsonObj JSON object to be populated.
+	 */
+	void StructuredDataToJsonObj(TSharedPtr< const FStructuredArchiveEntry > InElement, TSharedRef< FJsonObject > JsonObj );
+
+public:
+
+	static const FString TAG_FORMATVERSION;
+	static const FString TAG_NAMESPACE;
+	static const FString TAG_CHILDREN;
+	static const FString TAG_SUBNAMESPACES;
+	static const FString TAG_DEPRECATED_DEFAULTTEXT;
+	static const FString TAG_DEPRECATED_TRANSLATEDTEXT;
+	static const FString TAG_OPTIONAL;
+	static const FString TAG_SOURCE;
+	static const FString TAG_SOURCE_TEXT;
+	static const FString TAG_TRANSLATION;
+	static const FString TAG_TRANSLATION_TEXT;
+	static const FString TAG_METADATA_KEY;
+	static const FString NAMESPACE_DELIMITER;
+};
diff --git a/Engine/Source/Runtime/Internationalization/Public/Serialization/JsonInternationalizationManifestSerializer.h b/Engine/Source/Runtime/Internationalization/Public/Serialization/JsonInternationalizationManifestSerializer.h
new file mode 100644
--- /dev/null
+++ b/Engine/Source/Runtime/Internationalization/Public/Serialization/JsonInternationalizationManifestSerializer.h
@@ -0,0 +1,166 @@
+// Copyright 1998-2014 Epic Games, Inc. All Rights Reserved.
+#pragma once
+
+#include "IInternationalizationManifestSerializer.h"
+
+
+struct FStructuredEntry
+{
+public:
+	FStructuredEntry( const FString& InNamespace )
+		: Namespace( InNamespace )
+		, SubNamespaces()
+	{
+	}
+
+	const FString Namespace;
+	TArray< TSharedPtr< FStructuredEntry > > SubNamespaces;
+	TArray< TSharedPtr< class FManifestEntry > > ManifestEntries;
+};
+
+
+class FJsonObject;
+
+
+class INTERNATIONALIZATION_API FJsonInternationalizationManifestSerializer
+	: public IInternationalizationManifestSerializer
+{
+public:
+
+	/** Default constructor. */
+	FJsonInternationalizationManifestSerializer() { }
+
+	/**
+	 * Deserializes a Internationalization manifest from archive.
+	 *
+	 * @param Archive The archive to serialize from.
+	 * @param Manifest The populated Internationalization manifest.
+	 *
+	 * @return true if deserialization was successful, false otherwise.
+	 */
+	virtual bool DeserializeManifest( FArchive& Archive, TSharedRef< FInternationalizationManifest > Manifest ) override;
+
+	/**
+	 * Deserializes a Internationalization manifest from a string.
+	 *
+	 * @param InStr The string to serialize from.
+	 * @param Manifest The populated Internationalization manifest.
+	 *
+	 * @return true if deserialization was successful, false otherwise.
+	 */
+	virtual bool DeserializeManifest( const FString& InStr, TSharedRef< FInternationalizationManifest > Manifest );
+
+	/**
+	 * Deserializes a Internationalization manifest from a JSON object.
+	 *
+	 * @param InJsonObj The JSON object to serialize from.
+	 * @param Manifest The populated Internationalization manifest.
+	 *
+	 * @return true if deserialization was successful, false otherwise.
+	 */
+	virtual bool DeserializeManifest( TSharedRef< FJsonObject > InJsonObj, TSharedRef< FInternationalizationManifest > Manifest );
+
+	/**
+	 * Serializes a Internationalization manifest into an archive.
+	 *
+	 * @param Manifest The Internationalization manifest data to serialize.
+	 * @param Archive The archive to serialize into.
+	 *
+	 * @return true if serialization was successful, false otherwise.
+	 */
+	virtual bool SerializeManifest(TSharedRef< const FInternationalizationManifest > Manifest, FArchive& Archive ) override;
+
+	/**
+	 * Serializes a Internationalization manifest to a string.
+	 *
+	 * @param Manifest The Internationalization manifest data to serialize.
+	 * @param Str The string to serialize into.
+	 *
+	 * @return true if serialization was successful, false otherwise.
+	 */
+	virtual bool SerializeManifest( TSharedRef< const FInternationalizationManifest > Manifest, FString& Str );
+
+	/**
+	 * Serializes a Internationalization manifest to a JSON object.
+	 *
+	 * @param Manifest The Internationalization manifest data to serialize.
+	 * @param JsonObj The JSON object to serialize into.
+	 *
+	 * @return true if serialization was successful, false otherwise.
+	 */
+	virtual bool SerializeManifest( TSharedRef< const FInternationalizationManifest > Manifest, TSharedRef< FJsonObject > JsonObj );
+
+protected:
+
+	/**
+	 * Convert a JSON object to a Internationalization manifest.
+	 *
+	 * @param InJsonObj The JSON object to serialize from.
+	 * @param Manifest The Internationalization manifest that will store the data.
+	 *
+	 * @return true if deserialization was successful, false otherwise.
+	 */
+	bool DeserializeInternal( const TSharedRef< FJsonObject > InJsonObj, TSharedRef< FInternationalizationManifest > Manifest );
+
+	/**
+	 * Convert a Internationalization manifest to a JSON object.
+	 *
+	 * @param InManifest The Internationalization manifest object to serialize from.
+	 * @param JsonObj The Json object that will store the data.
+	 *
+	 * @return true if serialization was successful, false otherwise.
+	 */
+	bool SerializeInternal( TSharedRef< const FInternationalizationManifest > InManifest, TSharedRef< FJsonObject > JsonObj );
+
+	/**
+	 * Recursive function that will traverse the JSON object and populate a Internationalization manifest.
+	 *
+	 * @param InJsonObj The JSON object.
+	 * @param InNamespace The namespace of the parent JSON object.
+	 * @param Manifest The Internationalization manifest that will store the data.
+	 *
+	 * @return true if successful, false otherwise.
+	 */
+	bool JsonObjToManifest( TSharedRef< FJsonObject > InJsonObj, FString InNamespace, TSharedRef< FInternationalizationManifest > Manifest );
+
+	/**
+	 * Takes a Internationalization manifest and arranges the data into a hierarchy based on namespace.
+	 *
+	 * @param InManifest The Internationalization manifest.
+	 * @param RootElement The root element of the structured data.
+	 */
+	void GenerateStructuredData( TSharedRef< const FInternationalizationManifest > InManifest, TSharedPtr< FStructuredEntry > RootElement );
+
+	/**
+	 * Goes through the structured, hierarchy based, manifest data and does a non-culture specific sort on namespaces, default text, and key.
+	 *
+	 * @param RootElement The root element of the structured data.
+	 */
+	void SortStructuredData( TSharedPtr< FStructuredEntry > InElement );
+
+	/**
+	 * Populates a JSON object from Internationalization manifest data that has been structured based on namespace.
+	 *
+	 * @param InElement Internationalization manifest data structured based on namespace.
+	 * @param JsonObj JSON object to be populated.
+	 */
+	void StructuredDataToJsonObj( TSharedPtr< const FStructuredEntry> InElement, TSharedRef< FJsonObject > JsonObj );
+
+public:
+
+	static const FString TAG_FORMATVERSION;
+	static const FString TAG_NAMESPACE;
+	static const FString TAG_CHILDREN;
+	static const FString TAG_SUBNAMESPACES;
+	static const FString TAG_PATH;
+	static const FString TAG_OPTIONAL;
+	static const FString TAG_KEYCOLLECTION;
+	static const FString TAG_KEY;
+	static const FString TAG_DEPRECATED_DEFAULTTEXT;
+	static const FString TAG_SOURCE;
+	static const FString TAG_SOURCE_TEXT;
+	static const FString TAG_METADATA;
+	static const FString TAG_METADATA_INFO;
+	static const FString TAG_METADATA_KEY;
+	static const FString NAMESPACE_DELIMITER;
+};
diff --git a/Engine/Source/Runtime/Internationalization/Public/Serialization/JsonInternationalizationMetadataSerializer.h b/Engine/Source/Runtime/Internationalization/Public/Serialization/JsonInternationalizationMetadataSerializer.h
new file mode 100644
--- /dev/null
+++ b/Engine/Source/Runtime/Internationalization/Public/Serialization/JsonInternationalizationMetadataSerializer.h
@@ -0,0 +1,33 @@
+// Copyright 1998-2014 Epic Games, Inc. All Rights Reserved.
+#pragma once
+
+class FJsonObject;
+class FLocMetadataObject;
+
+
+class INTERNATIONALIZATION_API FJsonInternationalizationMetaDataSerializer
+{
+public:
+	/**
+	 * Deserializes manifest metadata from a JSON object.
+	 *
+	 * @param InJsonObj The JSON object to serialize from.
+	 * @param OutMetadataObj The resulting metadata object or null if the serialize is not successful.
+	 */
+	static void DeserializeMetadata( const TSharedRef< FJsonObject > JsonObj, TSharedPtr< FLocMetadataObject >& OutMetadataObj );
+
+	/**
+	 * Serializes manifest metadata object to a JSON object.
+	 *
+	 * @param Manifest The Internationalization metadata to serialize.
+	 * @param OutJsonObj The resulting JSON object or NULL if the serilize is not successful.
+	 */
+	static void SerializeMetadata(  const TSharedRef< FLocMetadataObject > Metadata, TSharedPtr< FJsonObject >& OutJsonObj );
+
+	/**
+	 * Utility function that will convert metadata to string using the JSON metadata serializers.
+	 *
+	 * @param Manifest The Internationalization metadata to conver to string.
+	 */
+	static FString MetadataToString( const TSharedPtr<FLocMetadataObject> Metadata );
+};
diff --git a/Engine/Source/Runtime/Json/Json.Build.cs b/Engine/Source/Runtime/Json/Json.Build.cs
new file mode 100644
--- /dev/null
+++ b/Engine/Source/Runtime/Json/Json.Build.cs
@@ -0,0 +1,23 @@
+// Copyright 1998-2014 Epic Games, Inc. All Rights Reserved.
+
+namespace UnrealBuildTool.Rules
+{
+	public class Json : ModuleRules
+	{
+		public Json(TargetInfo Target)
+		{
+			PublicDependencyModuleNames.AddRange(
+				new string[]
+				{
+					"Core",
+				}
+			); 
+
+			PrivateIncludePaths.AddRange(
+				new string[] {
+					"Runtime/Json/Private",
+				}
+			);
+		}
+	}
+}
diff --git a/Engine/Source/Runtime/Json/Private/Dom/JsonObject.cpp b/Engine/Source/Runtime/Json/Private/Dom/JsonObject.cpp
new file mode 100644
--- /dev/null
+++ b/Engine/Source/Runtime/Json/Private/Dom/JsonObject.cpp
@@ -0,0 +1,163 @@
+// Copyright 1998-2014 Epic Games, Inc. All Rights Reserved.
+
+#include "JsonPrivatePCH.h"
+
+
+void FJsonObject::SetField( const FString& FieldName, const TSharedPtr<FJsonValue>& Value )
+{
+	this->Values.Add(FieldName, Value);
+}
+
+
+void FJsonObject::RemoveField( const FString& FieldName )
+{
+	this->Values.Remove(FieldName);
+}
+
+
+double FJsonObject::GetNumberField( const FString& FieldName ) const
+{
+	return GetField<EJson::None>(FieldName)->AsNumber();
+}
+
+
+bool FJsonObject::TryGetNumberField( const FString& FieldName, double& OutNumber ) const
+{
+	TSharedPtr<FJsonValue> Field = TryGetField(FieldName);
+	return Field.IsValid() && Field->TryGetNumber(OutNumber);
+}
+
+
+bool FJsonObject::TryGetNumberField( const FString& FieldName, int32& OutNumber ) const
+{
+	TSharedPtr<FJsonValue> Field = TryGetField(FieldName);
+	return Field.IsValid() && Field->TryGetNumber(OutNumber);
+}
+
+
+bool FJsonObject::TryGetNumberField( const FString& FieldName, uint32& OutNumber ) const
+{
+	TSharedPtr<FJsonValue> Field = TryGetField(FieldName);
+	return Field.IsValid() && Field->TryGetNumber(OutNumber);
+}
+
+
+void FJsonObject::SetNumberField( const FString& FieldName, double Number )
+{
+	this->Values.Add(FieldName, MakeShareable(new FJsonValueNumber(Number)));
+}
+
+
+FString FJsonObject::GetStringField( const FString& FieldName ) const
+{
+	return GetField<EJson::None>(FieldName)->AsString();
+}
+
+
+bool FJsonObject::TryGetStringField( const FString& FieldName, FString& OutString ) const
+{
+	TSharedPtr<FJsonValue> Field = TryGetField(FieldName);
+	return Field.IsValid() && Field->TryGetString(OutString);
+}
+
+
+bool FJsonObject::TryGetStringArrayField( const FString& FieldName, TArray<FString>& OutArray ) const
+{
+	TSharedPtr<FJsonValue> Field = TryGetField(FieldName);
+
+	if (!Field.IsValid())
+	{
+		return false;
+	}
+
+	const TArray< TSharedPtr<FJsonValue> > *Array;
+
+	if (!Field->TryGetArray(Array))
+	{
+		return false;
+	}
+
+	for (int Idx = 0; Idx < Array->Num(); Idx++)
+	{
+		FString Element;
+
+		if (!(*Array)[Idx]->TryGetString(Element))
+		{
+			return false;
+		}
+
+		OutArray.Add(Element);
+	}
+
+	return true;
+}
+
+
+void FJsonObject::SetStringField( const FString& FieldName, const FString& StringValue )
+{
+	this->Values.Add(FieldName, MakeShareable(new FJsonValueString(StringValue)));
+}
+
+
+bool FJsonObject::GetBoolField( const FString& FieldName ) const
+{
+	return GetField<EJson::None>(FieldName)->AsBool();
+}
+
+
+bool FJsonObject::TryGetBoolField( const FString& FieldName, bool& OutBool ) const
+{
+	TSharedPtr<FJsonValue> Field = TryGetField(FieldName);
+	return Field.IsValid() && Field->TryGetBool(OutBool);
+}
+
+
+void FJsonObject::SetBoolField( const FString& FieldName, bool InValue )
+{
+	this->Values.Add(FieldName, MakeShareable( new FJsonValueBoolean(InValue)));
+}
+
+
+const TArray<TSharedPtr<FJsonValue>>& FJsonObject::GetArrayField( const FString& FieldName ) const
+{
+	return GetField<EJson::Array>(FieldName)->AsArray();
+}
+
+
+bool FJsonObject::TryGetArrayField(const FString& FieldName, const TArray< TSharedPtr<FJsonValue> >*& OutArray) const
+{
+	TSharedPtr<FJsonValue> Field = TryGetField(FieldName);
+	return Field.IsValid() && Field->TryGetArray(OutArray);
+}
+
+
+void FJsonObject::SetArrayField( const FString& FieldName, const TArray< TSharedPtr<FJsonValue> >& Array )
+{
+	this->Values.Add(FieldName, MakeShareable(new FJsonValueArray(Array)));
+}
+
+
+const TSharedPtr<FJsonObject>& FJsonObject::GetObjectField( const FString& FieldName ) const
+{
+	return GetField<EJson::Object>(FieldName)->AsObject();
+}
+
+
+bool FJsonObject::TryGetObjectField( const FString& FieldName, const TSharedPtr<FJsonObject>*& OutObject ) const
+{
+	TSharedPtr<FJsonValue> Field = TryGetField(FieldName);
+	return Field.IsValid() && Field->TryGetObject(OutObject);
+}
+
+
+void FJsonObject::SetObjectField( const FString& FieldName, const TSharedPtr<FJsonObject>& JsonObject )
+{
+	if (JsonObject.IsValid())
+	{
+		this->Values.Add(FieldName, MakeShareable(new FJsonValueObject(JsonObject.ToSharedRef())));
+	}
+	else
+	{
+		this->Values.Add(FieldName, MakeShareable( new FJsonValueNull()));
+	}
+}
diff --git a/Engine/Source/Runtime/Json/Private/Dom/JsonValue.cpp b/Engine/Source/Runtime/Json/Private/Dom/JsonValue.cpp
new file mode 100644
--- /dev/null
+++ b/Engine/Source/Runtime/Json/Private/Dom/JsonValue.cpp
@@ -0,0 +1,206 @@
+// Copyright 1998-2014 Epic Games, Inc. All Rights Reserved.
+
+#include "JsonPrivatePCH.h"
+
+
+const TArray< TSharedPtr<FJsonValue> > FJsonValue::EMPTY_ARRAY;
+const TSharedPtr<FJsonObject> FJsonValue::EMPTY_OBJECT(new FJsonObject());
+
+
+double FJsonValue::AsNumber() const
+{
+	double Number = 0.0;
+
+	if (!TryGetNumber(Number))
+	{
+		ErrorMessage(TEXT("Number"));
+	}
+
+	return Number;
+}
+
+
+FString FJsonValue::AsString() const 
+{
+	FString String;
+
+	if (!TryGetString(String))
+	{
+		ErrorMessage(TEXT("String"));
+	}
+
+	return String;
+}
+
+
+bool FJsonValue::AsBool() const 
+{
+	bool Bool = false;
+
+	if (!TryGetBool(Bool))
+	{
+		ErrorMessage(TEXT("Boolean")); 
+	}
+
+	return Bool;
+}
+
+
+const TArray< TSharedPtr<FJsonValue> >& FJsonValue::AsArray() const 
+{
+	const TArray<TSharedPtr<FJsonValue>> *Array = &EMPTY_ARRAY;
+
+	if (!TryGetArray(Array))
+	{
+		ErrorMessage(TEXT("Array")); 
+	}
+
+	return *Array;
+}
+
+
+const TSharedPtr<FJsonObject>& FJsonValue::AsObject() const 
+{
+	const TSharedPtr<FJsonObject> *Object = &EMPTY_OBJECT;
+
+	if (!TryGetObject(Object))
+	{
+		ErrorMessage(TEXT("Object"));
+	}
+
+	return *Object;
+}
+
+
+bool FJsonValue::TryGetNumber( int32& OutNumber ) const
+{
+	double Double;
+
+	if (TryGetNumber(Double) && (Double >= INT_MIN) && (Double <= INT_MAX))
+	{
+		if (Double >= 0.0)
+		{
+			OutNumber = (int32)(Double + 0.5);
+		}
+		else
+		{
+			OutNumber = (int32)(Double - 0.5);
+		}
+		
+		return true;
+	}
+
+	return false;
+}
+
+
+bool FJsonValue::TryGetNumber( uint32& OutNumber ) const
+{
+	double Double;
+
+	if (TryGetNumber(Double) && (Double >= 0.0) && (Double <= UINT_MAX))
+	{
+		OutNumber = (uint32)(Double + 0.5);
+
+		return true;
+	}
+
+	return false;
+}
+
+
+//static 
+bool FJsonValue::CompareEqual( const FJsonValue& Lhs, const FJsonValue& Rhs )
+{
+	if (Lhs.Type != Rhs.Type)
+	{
+		return false;
+	}
+
+	switch (Lhs.Type)
+	{
+	case EJson::None:
+	case EJson::Null:
+		return true;
+
+	case EJson::String:
+		return Lhs.AsString() == Rhs.AsString();
+
+	case EJson::Number:
+		return Lhs.AsNumber() == Rhs.AsNumber();
+
+	case EJson::Boolean:
+		return Lhs.AsBool() == Rhs.AsBool();
+
+	case EJson::Array:
+		{
+			const TArray< TSharedPtr<FJsonValue> >& LhsArray = Lhs.AsArray();
+			const TArray< TSharedPtr<FJsonValue> >& RhsArray = Rhs.AsArray();
+
+			if (LhsArray.Num() != RhsArray.Num())
+			{
+				return false;
+			}
+
+			// compare each element
+			for (int32 i = 0; i < LhsArray.Num(); ++i)
+			{
+				if (!CompareEqual(*LhsArray[i], *RhsArray[i]))
+				{
+					return false;
+				}
+			}
+		}
+		return true;
+
+	case EJson::Object:
+		{
+			const TSharedPtr<FJsonObject>& LhsObject = Lhs.AsObject();
+			const TSharedPtr<FJsonObject>& RhsObject = Rhs.AsObject();
+
+			if (LhsObject.IsValid() != RhsObject.IsValid())
+			{
+				return false;
+			}
+
+			if (LhsObject.IsValid())
+			{
+				if (LhsObject->Values.Num() != RhsObject->Values.Num())
+				{
+					return false;
+				}
+
+				// compare each element
+				for (const auto& It : LhsObject->Values)
+				{
+					const FString& Key = It.Key;
+					const TSharedPtr<FJsonValue>* RhsValue = RhsObject->Values.Find(Key);
+					if (RhsValue == NULL)
+					{
+						// not found in both objects
+						return false;
+					}
+
+					const TSharedPtr<FJsonValue>& LhsValue = It.Value;
+
+					if (LhsValue.IsValid() != RhsValue->IsValid())
+					{
+						return false;
+					}
+
+					if (LhsValue.IsValid())
+					{
+						if (!CompareEqual(*LhsValue.Get(), *RhsValue->Get()))
+						{
+							return false;
+						}
+					}
+				}
+			}
+		}
+		return true;
+
+	default:
+		return false;
+	}
+}
diff --git a/Engine/Source/Runtime/Json/Private/JsonModule.cpp b/Engine/Source/Runtime/Json/Private/JsonModule.cpp
new file mode 100644
--- /dev/null
+++ b/Engine/Source/Runtime/Json/Private/JsonModule.cpp
@@ -0,0 +1,30 @@
+// Copyright 1998-2014 Epic Games, Inc. All Rights Reserved.
+
+#include "JsonPrivatePCH.h"
+#include "ModuleManager.h"
+
+
+DEFINE_LOG_CATEGORY(LogJson);
+
+
+/**
+ * Implements the Json module.
+ */
+class FJsonModule
+	: public IModuleInterface
+{
+public:
+
+	// IModuleInterface interface
+
+	virtual void StartupModule( ) override { }
+	virtual void ShutdownModule( ) override { }
+
+	virtual bool SupportsDynamicReloading( ) override
+	{
+		return false;
+	}
+};
+
+
+IMPLEMENT_MODULE(FJsonModule, Json);
diff --git a/Engine/Source/Runtime/Json/Private/JsonPrivatePCH.h b/Engine/Source/Runtime/Json/Private/JsonPrivatePCH.h
new file mode 100644
--- /dev/null
+++ b/Engine/Source/Runtime/Json/Private/JsonPrivatePCH.h
@@ -0,0 +1,5 @@
+// Copyright 1998-2014 Epic Games, Inc. All Rights Reserved.
+
+#pragma once
+
+#include "Json.h"
diff --git a/Engine/Source/Runtime/Json/Public/Dom/JsonObject.h b/Engine/Source/Runtime/Json/Public/Dom/JsonObject.h
new file mode 100644
--- /dev/null
+++ b/Engine/Source/Runtime/Json/Public/Dom/JsonObject.h
@@ -0,0 +1,176 @@
+// Copyright 1998-2014 Epic Games, Inc. All Rights Reserved.
+
+#pragma once
+
+
+/**
+ * A Json Object is a structure holding an unordered set of name/value pairs.
+ * In a Json file, it is represented by everything between curly braces {}.
+ */
+class JSON_API FJsonObject
+{
+public:
+
+	TMap<FString, TSharedPtr<FJsonValue>> Values;
+
+	template<EJson JsonType>
+	TSharedPtr<FJsonValue> GetField( const FString& FieldName ) const
+	{
+		const TSharedPtr<FJsonValue>* Field = Values.Find(FieldName);
+		if ( Field != nullptr && Field->IsValid() )
+		{
+			if (JsonType == EJson::None || (*Field)->Type == JsonType)
+			{
+				return (*Field);
+			}
+			else
+			{
+				UE_LOG(LogJson, Warning, TEXT("Field %s is of the wrong type."), *FieldName);
+			}
+		}
+		else
+		{
+			UE_LOG(LogJson, Warning, TEXT("Field %s was not found."), *FieldName);
+		}
+
+		return TSharedPtr<FJsonValue>(new FJsonValueNull());
+	}
+
+	/**
+	 * Attempts to get the field with the specified name.
+	 *
+	 * @param FieldName The name of the field to get.
+	 * @return A pointer to the field, or nullptr if the field doesn't exist.
+	 */
+	TSharedPtr<FJsonValue> TryGetField( const FString& FieldName ) const
+	{
+		const TSharedPtr<FJsonValue>* Field = Values.Find(FieldName);
+		return (Field != nullptr && Field->IsValid()) ? *Field : TSharedPtr<FJsonValue>();
+	}
+
+	/**
+	 * Checks whether a field with the specified name exists in the object.
+	 *
+	 * @param FieldName The name of the field to check.
+	 * @return true if the field exists, false otherwise.
+	 */
+	bool HasField( const FString& FieldName) const
+	{
+		const TSharedPtr<FJsonValue>* Field = Values.Find(FieldName);
+		if(Field && Field->IsValid())
+		{
+			return true;
+		}
+
+		return false;
+	}
+	
+	/**
+	 * Checks whether a field with the specified name and type exists in the object.
+	 *
+	 * @param JsonType The type of the field to check.
+	 * @param FieldName The name of the field to check.
+	 * @return true if the field exists, false otherwise.
+	 */
+	template<EJson JsonType>
+	bool HasTypedField(const FString& FieldName) const
+	{
+		const TSharedPtr<FJsonValue>* Field = Values.Find(FieldName);
+		if(Field && Field->IsValid() && ((*Field)->Type == JsonType))
+		{
+			return true;
+		}
+
+		return false;
+	}
+
+	/**
+	 * Sets the value of the field with the specified name.
+	 *
+	 * @param FieldName The name of the field to set.
+	 * @param Value The value to set.
+	 */
+	void SetField( const FString& FieldName, const TSharedPtr<FJsonValue>& Value );
+
+	/**
+	 * Removes the field with the specified name.
+	 *
+	 * @param FieldName The name of the field to remove.
+	 */
+	void RemoveField(const FString& FieldName);
+
+	/**
+	 * Gets the field with the specified name as a number.
+	 *
+	 * Ensures that the field is present and is of type Json number.
+	 *
+	 * @param FieldName The name of the field to get.
+	 * @return The field's value as a number.
+	 */
+	double GetNumberField(const FString& FieldName) const;
+
+	/** Get the field named FieldName as a number. Returns false if it doesn't exist or cannot be converted. */
+	bool TryGetNumberField(const FString& FieldName, double& OutNumber) const;
+
+	/** Get the field named FieldName as a number, and makes sure it's within int32 range. Returns false if it doesn't exist or cannot be converted. */
+	bool TryGetNumberField(const FString& FieldName, int32& OutNumber) const;
+
+	/** Get the field named FieldName as a number, and makes sure it's within uint32 range. Returns false if it doesn't exist or cannot be converted.  */
+	bool TryGetNumberField(const FString& FieldName, uint32& OutNumber) const;
+
+	/** Add a field named FieldName with Number as value */
+	void SetNumberField( const FString& FieldName, double Number );
+
+	/** Get the field named FieldName as a string. */
+	FString GetStringField(const FString& FieldName) const;
+
+	/** Get the field named FieldName as a string. Returns false if it doesn't exist or cannot be converted. */
+	bool TryGetStringField(const FString& FieldName, FString& OutString) const;
+
+	/** Get the field named FieldName as an array of strings. Returns false if it doesn't exist or any member cannot be converted. */
+	bool TryGetStringArrayField(const FString& FieldName, TArray<FString>& OutArray) const;
+
+	/** Add a field named FieldName with value of StringValue */
+	void SetStringField( const FString& FieldName, const FString& StringValue );
+
+	/**
+	 * Gets the field with the specified name as a boolean.
+	 *
+	 * Ensures that the field is present and is of type Json number.
+	 *
+	 * @param FieldName The name of the field to get.
+	 * @return The field's value as a boolean.
+	 */
+	bool GetBoolField(const FString& FieldName) const;
+
+	/** Get the field named FieldName as a string. Returns false if it doesn't exist or cannot be converted. */
+	bool TryGetBoolField(const FString& FieldName, bool& OutBool) const;
+
+	/** Set a boolean field named FieldName and value of InValue */
+	void SetBoolField( const FString& FieldName, bool InValue );
+
+	/** Get the field named FieldName as an array. */
+	const TArray< TSharedPtr<FJsonValue> >& GetArrayField(const FString& FieldName) const;
+
+	/** Try to get the field named FieldName as an array, or return false if it's another type */
+	bool TryGetArrayField(const FString& FieldName, const TArray< TSharedPtr<FJsonValue> >*& OutArray) const;
+
+	/** Set an array field named FieldName and value of Array */
+	void SetArrayField( const FString& FieldName, const TArray< TSharedPtr<FJsonValue> >& Array );
+
+	/**
+	 * Gets the field with the specified name as a Json object.
+	 *
+	 * Ensures that the field is present and is of type Json number.
+	 *
+	 * @param FieldName The name of the field to get.
+	 * @return The field's value as a Json object.
+	 */
+	const TSharedPtr<FJsonObject>& GetObjectField(const FString& FieldName) const;
+
+	/** Try to get the field named FieldName as an object, or return false if it's another type */
+	bool TryGetObjectField(const FString& FieldName, const TSharedPtr<FJsonObject>*& OutObject) const;
+
+	/** Set an ObjectField named FieldName and value of JsonObject */
+	void SetObjectField( const FString& FieldName, const TSharedPtr<FJsonObject>& JsonObject );
+};
diff --git a/Engine/Source/Runtime/Json/Public/Dom/JsonValue.h b/Engine/Source/Runtime/Json/Public/Dom/JsonValue.h
new file mode 100644
--- /dev/null
+++ b/Engine/Source/Runtime/Json/Public/Dom/JsonValue.h
@@ -0,0 +1,167 @@
+// Copyright 1998-2014 Epic Games, Inc. All Rights Reserved.
+
+#pragma once
+
+
+class FJsonObject;
+
+
+/**
+ * A Json Value is a structure that can be any of the Json Types.
+ * It should never be used on its, only its derived types should be used.
+ */
+class JSON_API FJsonValue
+{
+public:
+
+	/** Returns this value as a double, logging an error and returning zero if this is not an Json Number */
+	double AsNumber() const;
+	
+	/** Returns this value as a number, logging an error and returning an empty string if not possible */
+	FString AsString() const;
+
+	/** Returns this value as a boolean, logging an error and returning false if not possible */
+	bool AsBool() const;
+
+	/** Returns this value as an array, logging an error and returning an empty array reference if not possible */
+	const TArray< TSharedPtr<FJsonValue> >& AsArray() const;
+
+	/** Returns this value as an object, throwing an error if this is not an Json Object */
+	virtual const TSharedPtr<FJsonObject>& AsObject() const;
+
+	/** Tries to convert this value to a number, returning false if not possible */
+	virtual bool TryGetNumber(double& OutDouble) const { return false; }
+
+	/** Tries to convert this value to an integer, returning false if not possible */
+	bool TryGetNumber(int32& OutNumber) const;
+
+	/** Tries to convert this value to an integer, returning false if not possible */
+	bool TryGetNumber(uint32& OutNumber) const;
+
+	/** Tries to convert this value to a string, returning false if not possible */
+	virtual bool TryGetString(FString& OutString) const { return false; }
+
+	/** Tries to convert this value to a bool, returning false if not possible */
+	virtual bool TryGetBool(bool& OutBool) const { return false; }
+
+	/** Tries to convert this value to an array, returning false if not possible */
+	virtual bool TryGetArray(const TArray< TSharedPtr<FJsonValue> >*& OutArray) const { return false; }
+
+	/** Tries to convert this value to an object, returning false if not possible */
+	virtual bool TryGetObject(const TSharedPtr<FJsonObject>*& Object) const { return false; }
+
+	/** Returns true if this value is a 'null' */
+	bool IsNull() const { return Type == EJson::Null || Type == EJson::None; }
+
+	/** Get a field of the same type as the argument */
+	void AsArgumentType(double                          & Value) { Value = AsNumber(); }
+	void AsArgumentType(FString                         & Value) { Value = AsString(); }
+	void AsArgumentType(bool                            & Value) { Value = AsBool  (); }
+	void AsArgumentType(TArray< TSharedPtr<FJsonValue> >& Value) { Value = AsArray (); }
+	void AsArgumentType(TSharedPtr<FJsonObject>         & Value) { Value = AsObject(); }
+
+	EJson Type;
+
+	static bool CompareEqual(const FJsonValue& Lhs, const FJsonValue& Rhs);
+	bool operator==(const FJsonValue& Rhs) const { return CompareEqual(*this, Rhs); }
+
+protected:
+
+	static const TArray< TSharedPtr<FJsonValue> > EMPTY_ARRAY;
+	static const TSharedPtr<FJsonObject> EMPTY_OBJECT;
+
+	FJsonValue() : Type(EJson::None) {}
+	virtual ~FJsonValue() {}
+
+	virtual FString GetType() const = 0;
+
+	void ErrorMessage(const FString& InType) const {UE_LOG(LogJson, Error, TEXT("Json Value of type '%s' used as a '%s'."), *GetType(), *InType);}
+};
+
+
+/** A Json String Value. */
+class JSON_API FJsonValueString : public FJsonValue
+{
+public:
+	FJsonValueString(const FString& InString) : Value(InString) {Type = EJson::String;}
+
+	virtual bool TryGetString(FString& OutString) const override	{ OutString = Value; return true; }
+	virtual bool TryGetNumber(double& OutDouble) const override		{ if(Value.IsNumeric()){ OutDouble = FCString::Atod(*Value); return true; } else { return false; } }
+	virtual bool TryGetBool(bool& OutBool) const override			{ OutBool = Value.ToBool(); return true; }
+
+protected:
+	FString Value;
+
+	virtual FString GetType() const {return TEXT("String");}
+};
+
+
+/** A Json Number Value. */
+class JSON_API FJsonValueNumber : public FJsonValue
+{
+public:
+	FJsonValueNumber(double InNumber) : Value(InNumber) {Type = EJson::Number;}
+	virtual bool TryGetNumber(double& OutNumber) const override		{ OutNumber = Value; return true; }
+	virtual bool TryGetBool(bool& OutBool) const override			{ OutBool = (Value != 0.0); return true; }
+	virtual bool TryGetString(FString& OutString) const override	{ OutString = FString::SanitizeFloat(Value); return true; }
+	
+protected:
+	double Value;
+
+	virtual FString GetType() const {return TEXT("Number");}
+};
+
+
+/** A Json Boolean Value. */
+class JSON_API FJsonValueBoolean : public FJsonValue
+{
+public:
+	FJsonValueBoolean(bool InBool) : Value(InBool) {Type = EJson::Boolean;}
+	virtual bool TryGetNumber(double& OutNumber) const override		{ OutNumber = Value ? 1 : 0; return true; }
+	virtual bool TryGetBool(bool& OutBool) const override			{ OutBool = Value; return true; }
+	virtual bool TryGetString(FString& OutString) const override	{ OutString = Value ? TEXT("true") : TEXT("false"); return true; }
+	
+protected:
+	bool Value;
+
+	virtual FString GetType() const {return TEXT("Boolean");}
+};
+
+
+/** A Json Array Value. */
+class JSON_API FJsonValueArray : public FJsonValue
+{
+public:
+	FJsonValueArray(const TArray< TSharedPtr<FJsonValue> >& InArray) : Value(InArray) {Type = EJson::Array;}
+	virtual bool TryGetArray(const TArray< TSharedPtr<FJsonValue> >*& OutArray) const override	{ OutArray = &Value; return true; }
+	
+protected:
+	TArray< TSharedPtr<FJsonValue> > Value;
+
+	virtual FString GetType() const {return TEXT("Array");}
+};
+
+
+/** A Json Object Value. */
+class JSON_API FJsonValueObject : public FJsonValue
+{
+public:
+	FJsonValueObject(TSharedPtr<FJsonObject> InObject) : Value(InObject) {Type = EJson::Object;}
+	virtual bool TryGetObject(const TSharedPtr<FJsonObject>*& OutObject) const override			{ OutObject = &Value; return true; }
+	
+protected:
+	TSharedPtr<FJsonObject> Value;
+
+	virtual FString GetType() const {return TEXT("Object");}
+};
+
+
+/** A Json Null Value. */
+class JSON_API FJsonValueNull : public FJsonValue
+{
+public:
+	FJsonValueNull() {Type = EJson::Null;}
+
+protected:
+	virtual FString GetType() const {return TEXT("Null");}
+};
diff --git a/Engine/Source/Runtime/Json/Public/Json.h b/Engine/Source/Runtime/Json/Public/Json.h
new file mode 100644
--- /dev/null
+++ b/Engine/Source/Runtime/Json/Public/Json.h
@@ -0,0 +1,27 @@
+// Copyright 1998-2014 Epic Games, Inc. All Rights Reserved.
+
+#pragma once
+
+
+/* Public dependencies
+ *****************************************************************************/
+
+#include "Core.h"
+
+
+/* Public includes
+ *****************************************************************************/
+
+DECLARE_LOG_CATEGORY_EXTERN(LogJson, Log, All);
+
+#include "JsonPrintPolicy.h"
+#include "PrettyJsonPrintPolicy.h"
+#include "CondensedJsonPrintPolicy.h"
+
+#include "JsonTypes.h"
+#include "JsonValue.h"
+#include "JsonObject.h"
+
+#include "JsonReader.h"
+#include "JsonWriter.h"
+#include "JsonSerializer.h"
diff --git a/Engine/Source/Runtime/Json/Public/Policies/CondensedJsonPrintPolicy.h b/Engine/Source/Runtime/Json/Public/Policies/CondensedJsonPrintPolicy.h
new file mode 100644
--- /dev/null
+++ b/Engine/Source/Runtime/Json/Public/Policies/CondensedJsonPrintPolicy.h
@@ -0,0 +1,18 @@
+// Copyright 1998-2014 Epic Games, Inc. All Rights Reserved.
+
+#pragma once
+
+
+/**
+ * Template for print policies that generate compressed output.
+ *
+ * @param CharType The type of characters to print, i.e. TCHAR or ANSICHAR.
+ */
+template <class CharType>
+struct TCondensedJsonPrintPolicy
+	: public TJsonPrintPolicy<CharType>
+{
+	static inline void WriteLineTerminator(FArchive* Stream) {}
+	static inline void WriteTabs(FArchive* Stream, int32 Count) {}
+	static inline void WriteSpace(FArchive* Stream) {}
+};
diff --git a/Engine/Source/Runtime/Json/Public/Policies/JsonPrintPolicy.h b/Engine/Source/Runtime/Json/Public/Policies/JsonPrintPolicy.h
new file mode 100644
--- /dev/null
+++ b/Engine/Source/Runtime/Json/Public/Policies/JsonPrintPolicy.h
@@ -0,0 +1,50 @@
+// Copyright 1998-2014 Epic Games, Inc. All Rights Reserved.
+
+#pragma once
+
+
+/**
+ * Base template for Json print policies.
+ *
+ * @param CharType The type of characters to print, i.e. TCHAR or ANSICHAR.
+ */
+template <class CharType>
+struct TJsonPrintPolicy
+{
+	/**
+	 * Writes a single character to the output stream.
+	 *
+	 * @param Stream The stream to write to.
+	 * @param Char The character to write.
+	 */
+	static inline void WriteChar( FArchive* Stream, CharType Char )
+	{
+		Stream->Serialize(&Char, sizeof(CharType));
+	}
+
+	/**
+	 * Writes a string to the output stream.
+	 *
+	 * @param Stream The stream to write to.
+	 * @param String The string to write.
+	 */
+	static inline void WriteString( FArchive* Stream, const FString& String )
+	{
+		const TCHAR* CharPtr = *String;
+
+		for (int32 CharIndex = 0; CharIndex < String.Len(); ++CharIndex, ++CharPtr)
+		{
+			WriteChar(Stream, *CharPtr);
+		}
+	}
+};
+
+
+/**
+ * Specialization for TCHAR that allows direct copying from FString data.
+ */
+template <>
+inline void TJsonPrintPolicy<TCHAR>::WriteString( FArchive* Stream, const FString& String )
+{
+	Stream->Serialize((void*)*String, String.Len() * sizeof(TCHAR));
+}
diff --git a/Engine/Source/Runtime/Json/Public/Policies/PrettyJsonPrintPolicy.h b/Engine/Source/Runtime/Json/Public/Policies/PrettyJsonPrintPolicy.h
new file mode 100644
--- /dev/null
+++ b/Engine/Source/Runtime/Json/Public/Policies/PrettyJsonPrintPolicy.h
@@ -0,0 +1,34 @@
+// Copyright 1998-2014 Epic Games, Inc. All Rights Reserved.
+
+#pragma once
+
+
+/**
+ * Template for print policies that generate human readable output.
+ *
+ * @param CharType The type of characters to print, i.e. TCHAR or ANSICHAR.
+ */
+template <class CharType>
+struct TPrettyJsonPrintPolicy
+	: public TJsonPrintPolicy<CharType>
+{
+	static inline void WriteLineTerminator( FArchive* Stream )
+	{
+		TJsonPrintPolicy<CharType>::WriteString(Stream, LINE_TERMINATOR);
+	}
+
+	static inline void WriteTabs( FArchive* Stream, int32 Count )
+	{
+		CharType Tab = CharType('\t');
+
+		for (int32 i = 0; i < Count; ++i)
+		{
+			TJsonPrintPolicy<CharType>::WriteChar(Stream, Tab);
+		}
+	}
+
+	static inline void WriteSpace( FArchive* Stream )
+	{
+		TJsonPrintPolicy<CharType>::WriteChar(Stream, CharType(' '));
+	}
+};
diff --git a/Engine/Source/Runtime/Json/Public/Serialization/JsonReader.h b/Engine/Source/Runtime/Json/Public/Serialization/JsonReader.h
new file mode 100644
--- /dev/null
+++ b/Engine/Source/Runtime/Json/Public/Serialization/JsonReader.h
@@ -0,0 +1,831 @@
+// Copyright 1998-2014 Epic Games, Inc. All Rights Reserved.
+
+#pragma once
+
+
+static EJsonNotation TokenToNotationTable[] = 
+{ 
+	EJsonNotation::Error,			/*EJsonToken::None*/
+	EJsonNotation::Error,			/*EJsonToken::Comma*/
+	EJsonNotation::ObjectStart,		/*EJsonToken::CurlyOpen*/
+	EJsonNotation::ObjectEnd,		/*EJsonToken::CurlyClose*/
+	EJsonNotation::ArrayStart,		/*EJsonToken::SquareOpen*/
+	EJsonNotation::ArrayEnd,		/*EJsonToken::SquareClose*/
+	EJsonNotation::Error,			/*EJsonToken::Colon*/
+	EJsonNotation::String,			/*EJsonToken::String*/
+	EJsonNotation::Number,			/*EJsonToken::Number*/
+	EJsonNotation::Boolean,			/*EJsonToken::True*/
+	EJsonNotation::Boolean,			/*EJsonToken::False*/
+	EJsonNotation::Null,			/*EJsonToken::Null*/
+};
+
+
+template <class CharType = TCHAR>
+class TJsonReader
+{
+public:
+
+	static TSharedRef< TJsonReader<CharType> > Create( FArchive* const Stream )
+	{
+		return MakeShareable( new TJsonReader<CharType>( Stream ) );
+	}
+
+public:
+
+	virtual ~TJsonReader() {}
+
+	bool ReadNext( EJsonNotation& Notation )
+	{
+		if (!ErrorMessage.IsEmpty())
+		{
+			Notation = EJsonNotation::Error;
+			return false;
+		}
+
+		if (Stream == nullptr)
+		{
+			Notation = EJsonNotation::Error;
+			SetErrorMessage(TEXT("Null Stream"));
+			return true;
+		}
+
+		const bool AtEndOfStream = Stream->AtEnd();
+
+		if (AtEndOfStream && !FinishedReadingRootObject)
+		{
+			Notation = EJsonNotation::Error;
+			SetErrorMessage(TEXT("Improperly formatted."));
+			return true;
+		}
+
+		if (FinishedReadingRootObject && !AtEndOfStream)
+		{
+			Notation = EJsonNotation::Error;
+			SetErrorMessage(TEXT("Unexpected additional input found."));
+			return true;
+		}
+
+		if (AtEndOfStream)
+		{
+			return false;
+		}
+
+		bool ReadWasSuccess = false;
+		Identifier.Empty();
+
+		do
+		{
+			EJson CurrentState = EJson::None;
+
+			if (ParseState.Num() > 0)
+			{
+				CurrentState = ParseState.Top();
+			}
+
+			switch (CurrentState)
+			{
+				case EJson::Array:
+					ReadWasSuccess = ReadNextArrayValue( /*OUT*/ CurrentToken );
+					break;
+
+				case EJson::Object:
+					ReadWasSuccess = ReadNextObjectValue( /*OUT*/ CurrentToken );
+					break;
+
+				default:
+					ReadWasSuccess = ReadStart( /*OUT*/ CurrentToken );
+					break;
+			}
+		}
+		while (ReadWasSuccess && (CurrentToken == EJsonToken::None));
+
+		Notation = TokenToNotationTable[(int32)CurrentToken];
+		FinishedReadingRootObject = ParseState.Num() == 0;
+
+		if (!ReadWasSuccess || (Notation == EJsonNotation::Error))
+		{
+			Notation = EJsonNotation::Error;
+
+			if (ErrorMessage.IsEmpty())
+			{
+				SetErrorMessage(TEXT("Unknown Error Occurred"));
+			}
+
+			return true;
+		}
+
+		if (FinishedReadingRootObject && !Stream->AtEnd())
+		{
+			ParseWhiteSpace();
+		}
+
+		return ReadWasSuccess;
+	}
+
+	bool SkipObject()
+	{
+		return ReadUntilMatching(EJsonNotation::ObjectEnd);
+	}
+
+	bool SkipArray()
+	{
+		return ReadUntilMatching(EJsonNotation::ArrayEnd);
+	}
+
+	FORCEINLINE const FString& GetIdentifier() const { return Identifier; }
+
+	FORCEINLINE const FString& GetValueAsString() const 
+	{ 
+		check(CurrentToken == EJsonToken::String);
+		return StringValue;
+	}
+	
+	FORCEINLINE double GetValueAsNumber() const 
+	{ 
+		check(CurrentToken == EJsonToken::Number);
+		return NumberValue;
+	}
+	
+	FORCEINLINE bool GetValueAsBoolean() const 
+	{ 
+		check((CurrentToken == EJsonToken::True) || (CurrentToken == EJsonToken::False));
+		return BoolValue; 
+	}
+
+	FORCEINLINE const FString& GetErrorMessage() const
+	{
+		return ErrorMessage;
+	}
+
+	FORCEINLINE const uint32 GetLineNumber() const
+	{
+		return LineNumber;
+	}
+
+	FORCEINLINE const uint32 GetCharacterNumber() const
+	{
+		return CharacterNumber;
+	}
+
+protected:
+
+	/** Hidden default constructor. */
+	TJsonReader()
+		: ParseState()
+		, CurrentToken( EJsonToken::None )
+		, Stream( nullptr )
+		, Identifier()
+		, ErrorMessage()
+		, StringValue()
+		, NumberValue( 0.0f )
+		, LineNumber( 1 )
+		, CharacterNumber( 0 )
+		, BoolValue( false )
+		, FinishedReadingRootObject( false )
+	{ }
+
+	/**
+	 * Creates and initializes a new instance with the given input.
+	 *
+	 * @param InStream An archive containing the input.
+	 */
+	TJsonReader(FArchive* InStream)
+		: ParseState()
+		, CurrentToken(EJsonToken::None)
+		, Stream(InStream)
+		, Identifier()
+		, ErrorMessage()
+		, StringValue()
+		, NumberValue(0.0f)
+		, LineNumber(1)
+		, CharacterNumber(0)
+		, BoolValue(false)
+		, FinishedReadingRootObject(false)
+	{ }
+
+private:
+
+	void SetErrorMessage( const FString& Message )
+	{
+		ErrorMessage = Message + FString::Printf(TEXT(" Line: %u Ch: %u"), LineNumber, CharacterNumber);
+	}
+
+	bool ReadUntilMatching( const EJsonNotation ExpectedNotation )
+	{
+		uint32 ScopeCount = 0;
+		EJsonNotation Notation;
+
+		while (ReadNext(Notation))
+		{
+			if ((ScopeCount == 0) && (Notation == ExpectedNotation))
+			{
+				return true;
+			}
+
+			switch (Notation)
+			{
+			case EJsonNotation::ObjectStart:
+			case EJsonNotation::ArrayStart:
+				++ScopeCount;
+				break;
+
+			case EJsonNotation::ObjectEnd:
+			case EJsonNotation::ArrayEnd:
+				--ScopeCount;
+				break;
+
+			case EJsonNotation::Error:
+				return false;
+				break;
+			}
+		}
+
+		return true;
+	}
+
+	bool ReadStart( EJsonToken& Token )
+	{
+		ParseWhiteSpace();
+
+		Token = EJsonToken::None;
+
+		if (NextToken(Token) == false)
+		{
+			return false;
+		}
+
+		if ((Token != EJsonToken::CurlyOpen) && (Token != EJsonToken::SquareOpen))
+		{
+			SetErrorMessage(TEXT("Open Curly or Square Brace token expected, but not found."));
+			return false;
+		}
+
+		return true;
+	}
+
+	bool ReadNextObjectValue( EJsonToken& Token )
+	{
+		const bool bCommaPrepend = Token != EJsonToken::CurlyOpen;
+		Token = EJsonToken::None;
+
+		if (NextToken(Token) == false)
+		{
+			return false;
+		}
+
+		if (Token == EJsonToken::CurlyClose)
+		{
+			return true;
+		}
+		else
+		{
+			if (bCommaPrepend)
+			{
+				if (Token != EJsonToken::Comma)
+				{
+					SetErrorMessage( TEXT("Comma token expected, but not found.") );
+					return false;
+				}
+
+				Token = EJsonToken::None;
+
+				if (!NextToken(Token))
+				{
+					return false;
+				}
+			}
+
+			if (Token != EJsonToken::String)
+			{
+				SetErrorMessage( TEXT("String token expected, but not found.") );
+				return false;
+			}
+
+			Identifier = StringValue;
+			Token = EJsonToken::None;
+
+			if (!NextToken(Token))
+			{
+				return false;
+			}
+
+			if (Token != EJsonToken::Colon)
+			{
+				SetErrorMessage( TEXT("Colon token expected, but not found.") );
+				return false;
+			}
+
+			Token = EJsonToken::None;
+
+			if (!NextToken(Token))
+			{
+				return false;
+			}
+		}
+
+		return true;
+	}
+
+	bool ReadNextArrayValue( EJsonToken& Token )
+	{
+		const bool bCommaPrepend = Token != EJsonToken::SquareOpen;
+
+		Token = EJsonToken::None;
+
+		if (!NextToken(Token))
+		{
+			return false;
+		}
+
+		if (Token == EJsonToken::SquareClose)
+		{
+			return true;
+		}
+		else
+		{
+			if (bCommaPrepend)
+			{
+				if (Token != EJsonToken::Comma)
+				{
+					SetErrorMessage( TEXT("Comma token expected, but not found.") );
+					return false;
+				}
+
+				Token = EJsonToken::None;
+
+				if (!NextToken(Token))
+				{
+					return false;
+				}
+			}
+		}
+
+		return true;
+	}
+
+	bool NextToken( EJsonToken& OutToken )
+	{
+		while (!Stream->AtEnd())
+		{
+			CharType Char;
+			Stream->Serialize(&Char, sizeof(CharType));
+			++CharacterNumber;
+
+			if (Char == CharType('\0'))
+			{
+				break;
+			}
+
+			if (IsLineBreak(Char))
+			{
+				++LineNumber;
+				CharacterNumber = 0;
+			}
+
+			if (!IsWhitespace(Char))
+			{
+				if (IsJsonNumber(Char))
+				{
+					if (!ParseNumberToken(Char))
+					{
+						return false;
+					}
+
+					OutToken = EJsonToken::Number;
+					return true;
+				}
+
+				switch (Char)
+				{
+				case CharType('{'):
+					OutToken = EJsonToken::CurlyOpen; ParseState.Push( EJson::Object );
+					return true;
+
+				case CharType('}'):
+					OutToken = EJsonToken::CurlyClose; ParseState.Pop();
+					return true;
+
+				case CharType('['):
+					OutToken = EJsonToken::SquareOpen; ParseState.Push( EJson::Array );
+					return true;
+
+				case CharType(']'):
+					OutToken = EJsonToken::SquareClose; ParseState.Pop();
+					return true;
+
+				case CharType(':'):
+					OutToken = EJsonToken::Colon;
+					return true;
+
+				case CharType(','):
+					OutToken = EJsonToken::Comma;
+					return true;
+
+				case CharType('\"'):
+					{
+						if (!ParseStringToken())
+						{
+							return false;
+						}
+
+						OutToken = EJsonToken::String;
+					}
+					return true;
+
+				case CharType('t'): case CharType('T'):
+				case CharType('f'): case CharType('F'):
+				case CharType('n'): case CharType('N'):
+					{
+						FString Test;
+						Test += Char;
+
+						while (!Stream->AtEnd())
+						{
+							Stream->Serialize(&Char, sizeof(CharType));
+
+							if (IsAlphaNumber(Char))
+							{
+								++CharacterNumber;
+								Test += Char;
+							}
+							else
+							{
+								// backtrack and break
+								Stream->Seek(Stream->Tell() - sizeof(CharType));
+								break;
+							}
+						}
+
+						if (Test == TEXT("False"))
+						{
+							BoolValue = false;
+							OutToken = EJsonToken::False;
+							return true;
+						}
+
+						if (Test == TEXT("True"))
+						{
+							BoolValue = true;
+							OutToken = EJsonToken::True;
+							return true;
+						}
+
+						if (Test == TEXT("Null"))
+						{
+							OutToken = EJsonToken::Null;
+							return true;
+						}
+
+						SetErrorMessage( TEXT("Invalid Json Token. Check that your member names have quotes around them!") );
+						return false;
+					}
+
+				default: 
+					SetErrorMessage( TEXT("Invalid Json Token.") );
+					return false;
+				}
+			}
+		}
+
+		SetErrorMessage( TEXT("Invalid Json Token.") );
+		return false;
+	}
+
+	bool ParseStringToken()
+	{
+		FString String;
+
+		while (true)
+		{
+			if (Stream->AtEnd())
+			{
+				SetErrorMessage( TEXT("String Token Abruptly Ended.") );
+				return false;
+			}
+
+			CharType Char;
+			Stream->Serialize(&Char, sizeof(CharType));
+			++CharacterNumber;
+
+			if (Char == CharType('\"'))
+			{
+				break;
+			}
+
+			if (Char == CharType('\\'))
+			{
+				Stream->Serialize(&Char, sizeof(CharType));
+				++CharacterNumber;
+
+				switch (Char)
+				{
+				case CharType('\"'): case CharType('\\'): case CharType('/'): String += Char; break;
+				case CharType('f'): String += CharType('\f'); break;
+				case CharType('r'): String += CharType('\r'); break;
+				case CharType('n'): String += CharType('\n'); break;
+				case CharType('b'): String += CharType('\b'); break;
+				case CharType('t'): String += CharType('\t'); break;
+				case CharType('u'):
+					// 4 hex digits, like \uAB23, which is a 16 bit number that we would usually see as 0xAB23
+					{
+						int32 HexNum = 0;
+
+						for (int32 Radix = 3; Radix >= 0; --Radix)
+						{
+							if (Stream->AtEnd())
+							{
+								SetErrorMessage( TEXT("String Token Abruptly Ended.") );
+								return false;
+							}
+
+							Stream->Serialize(&Char, sizeof(CharType));
+							++CharacterNumber;
+
+							int32 HexDigit = FParse::HexDigit(Char);
+
+							if ((HexDigit == 0) && (Char != CharType('0')))
+							{
+								SetErrorMessage( TEXT("Invalid Hexadecimal digit parsed.") );
+								return false;
+							}
+
+							HexNum += HexDigit * FMath::Pow(16, Radix);
+						}
+
+						String += (CharType)HexNum;
+					}
+					break;
+
+				default:
+					SetErrorMessage( TEXT("Bad Json escaped char.") );
+					return false;
+				}
+			}
+			else
+			{
+				String += Char;
+			}
+		}
+
+		StringValue = MoveTemp(String);
+		return true;
+	}
+
+	bool ParseNumberToken( CharType FirstChar )
+	{
+		FString String;
+		int32 State = 0;
+		bool UseFirstChar = true;
+		bool Error = false;
+
+		while (true)
+		{
+			CharType Char;
+			if (UseFirstChar)
+			{
+				Char = FirstChar;
+				UseFirstChar = false;
+			}
+			else
+			{
+				Stream->Serialize(&Char, sizeof(CharType));
+				++CharacterNumber;
+			}
+
+			// The following code doesn't actually derive the Json Number: that is handled
+			// by the function FCString::Atof below. This code only ensures the Json Number is
+			// EXACTLY to specification
+			if (IsJsonNumber(Char))
+			{
+				// ensure number follows Json format before converting
+				// This switch statement is derived from a finite state automata
+				// derived from the Json spec. A table was not used for simplicity.
+				switch (State)
+				{
+				case 0:
+					if (Char == CharType('-')) { State = 1; }
+					else if (Char == CharType('0')) { State = 2; }
+					else if (IsNonZeroDigit(Char)) { State = 3; }
+					else { Error = true; }
+					break;
+
+				case 1:
+					if (Char == CharType('0')) { State = 2; }
+					else if (IsNonZeroDigit(Char)) { State = 3; }
+					else { Error = true; }
+					break;
+
+				case 2:
+					if (Char == CharType('.')) { State = 4; }
+					else if (Char == CharType('e') || Char == CharType('E')) { State = 5; }
+					else { Error = true; }
+					break;
+
+				case 3:
+					if (IsDigit(Char)) { State = 3; }
+					else if (Char == CharType('.')) { State = 4; }
+					else if (Char == CharType('e') || Char == CharType('E')) { State = 5; }
+					else { Error = true; }
+					break;
+
+				case 4:
+					if (IsDigit(Char)) { State = 6; }
+					else { Error = true; }
+					break;
+
+				case 5:
+					if (Char == CharType('-') ||Char == CharType('+')) { State = 7; }
+					else if (IsDigit(Char)) { State = 8; }
+					else {Error = true;}
+					break;
+
+				case 6:
+					if (IsDigit(Char)) { State = 6; }
+					else if (Char == CharType('e') || Char == CharType('E')) { State = 5; }
+					else { Error = true; }
+					break;
+
+				case 7:
+					if (IsDigit(Char)) { State = 8; }
+					else { Error = true; }
+					break;
+
+				case 8:
+					if (IsDigit(Char)) { State = 8; }
+					else { Error = true; }
+					break;
+
+				default:
+					SetErrorMessage( TEXT("Unknown state reached in Json Number Token.") );
+					return false;
+				}
+
+				if (Error)
+				{
+					break;
+				}
+
+				String += Char;
+			}
+			else
+			{
+				// backtrack once because we read a non-number character
+				Stream->Seek(Stream->Tell() - sizeof(CharType));
+				--CharacterNumber;
+				// and now the number is fully tokenized
+				break;
+			}
+		}
+
+		// ensure the number has followed valid Json format
+		if (!Error && ((State == 2) || (State == 3) || (State == 6) || (State == 8)))
+		{
+			NumberValue = FCString::Atod(*String);
+			return true;
+		}
+
+		SetErrorMessage( TEXT("Poorly formed Json Number Token.") );
+		return false;
+	}
+
+	void ParseWhiteSpace()
+	{
+		while (!Stream->AtEnd())
+		{
+			CharType Char;
+			Stream->Serialize(&Char, sizeof(CharType));
+			++CharacterNumber;
+
+			if (IsLineBreak(Char))
+			{
+				++LineNumber;
+				CharacterNumber = 0;
+			}
+
+			if (!IsWhitespace(Char))
+			{
+				// backtrack and break
+				Stream->Seek(Stream->Tell() - sizeof(CharType));
+				--CharacterNumber;
+				break;
+			}
+		}
+	}
+
+	bool IsLineBreak( const CharType& Char )
+	{
+		return Char == CharType('\n');
+	}
+
+	/** Can't use FChar::IsWhitespace because it is TCHAR specific, and it doesn't handle newlines */
+	bool IsWhitespace( const CharType& Char )
+	{
+		return Char == CharType(' ') || Char == CharType('\t') || Char == CharType('\n') || Char == CharType('\r');
+	}
+
+	/** Can't use FChar::IsDigit because it is TCHAR specific, and it doesn't handle all the other Json number characters */
+	bool IsJsonNumber( const CharType& Char )
+	{
+		return ((Char >= CharType('0') && Char <= CharType('9')) ||
+			Char == CharType('-') || Char == CharType('.') || Char == CharType('+') || Char == CharType('e') || Char == CharType('E'));
+	}
+
+	/** Can't use FChar::IsDigit because it is TCHAR specific */
+	bool IsDigit( const CharType& Char )
+	{
+		return (Char >= CharType('0') && Char <= CharType('9'));
+	}
+
+	bool IsNonZeroDigit( const CharType& Char )
+	{
+		return (Char >= CharType('1') && Char <= CharType('9'));
+	}
+
+	/** Can't use FChar::IsAlpha because it is TCHAR specific. Also, this only checks A through Z (no underscores or other characters). */
+	bool IsAlphaNumber( const CharType& Char )
+	{
+		return (Char >= CharType('a') && Char <= CharType('z')) || (Char >= CharType('A') && Char <= CharType('Z'));
+	}
+
+protected:
+
+	TArray<EJson> ParseState;
+	EJsonToken CurrentToken;
+
+	FArchive* Stream;
+	FString Identifier;
+	FString ErrorMessage;
+	FString StringValue;
+	double NumberValue;
+	uint32 LineNumber;
+	uint32 CharacterNumber;
+	bool BoolValue;
+	bool FinishedReadingRootObject;
+};
+
+
+class FJsonStringReader
+	: public TJsonReader<TCHAR>
+{
+public:
+
+	static TSharedRef<FJsonStringReader> Create( const FString& JsonString )
+	{
+		return MakeShareable(new FJsonStringReader(JsonString));
+	}
+
+public:
+
+	virtual ~FJsonStringReader()
+	{
+		if (Reader != nullptr)
+		{
+			check(Reader->Close());
+			delete Reader;
+		}
+	}
+
+private:
+
+	/**
+	 * Parses a string containing Json information.
+	 *
+	 * @param JsonString The Json string to parse.
+	 */
+	FJsonStringReader( const FString& JsonString )
+		: Content(JsonString)
+		, Reader(nullptr)
+	{
+		if (Content.IsEmpty())
+		{
+			return;
+		}
+
+		Reader = new FBufferReader((void*)*Content, Content.Len() * sizeof(TCHAR), false);
+		check(Reader);
+
+		Stream = Reader;
+	}
+
+private:
+
+	const FString Content;
+	FBufferReader* Reader;
+};
+
+
+template <class CharType = TCHAR>
+class TJsonReaderFactory
+{
+public:
+
+	static TSharedRef<TJsonReader<TCHAR>> Create( const FString& JsonString )
+	{
+		return FJsonStringReader::Create(JsonString);
+	}
+
+	static TSharedRef<TJsonReader<CharType>> Create( FArchive* const Stream )
+	{
+		return TJsonReader<CharType>::Create(Stream);
+	}
+};
diff --git a/Engine/Source/Runtime/Json/Public/Serialization/JsonSerializer.h b/Engine/Source/Runtime/Json/Public/Serialization/JsonSerializer.h
new file mode 100644
--- /dev/null
+++ b/Engine/Source/Runtime/Json/Public/Serialization/JsonSerializer.h
@@ -0,0 +1,361 @@
+// Copyright 1998-2014 Epic Games, Inc. All Rights Reserved.
+
+#pragma once
+
+
+class FJsonSerializer
+{
+public:
+
+	template <class CharType>
+	static bool Deserialize( const TSharedRef< TJsonReader<CharType> >& Reader, TArray< TSharedPtr<FJsonValue> >& OutArray )
+	{
+		StackState State;
+		if ( !Deserialize( Reader, /*OUT*/State ) )
+		{
+			return false;
+		}
+
+		if ( State.Object.IsValid() )
+		{
+			return false;
+		}
+
+		OutArray = State.Array;
+
+		return true;
+	}
+
+	template <class CharType>
+	static bool Deserialize( const TSharedRef< TJsonReader<CharType> >& Reader, TSharedPtr<FJsonObject>& OutObject )
+	{
+		StackState State;
+		if ( !Deserialize( Reader, /*OUT*/State ) )
+		{
+			return false;
+		}
+
+		if ( !State.Object.IsValid() )
+		{
+			return false;
+		}
+
+		OutObject = State.Object;
+
+		return true;
+	}
+
+	template <class CharType, class PrintPolicy >
+	static bool Serialize( const TArray< TSharedPtr<FJsonValue> >& Array, const TSharedRef< TJsonWriter< CharType, PrintPolicy > >& Writer )
+	{
+		TSharedRef< FElement > StartingElement = MakeShareable( new FElement( Array ) );
+		return FJsonSerializer::Serialize<CharType, PrintPolicy>( StartingElement, Writer );
+	}
+
+	template <class CharType, class PrintPolicy >
+	static bool Serialize( const TSharedRef< FJsonObject >& Object, const TSharedRef< TJsonWriter< CharType, PrintPolicy > >& Writer )
+	{
+		TSharedRef< FElement > StartingElement = MakeShareable( new FElement( Object ) );
+		return FJsonSerializer::Serialize<CharType, PrintPolicy>( StartingElement, Writer );
+	}
+
+	template <class CharType, class PrintPolicy >
+	static bool Serialize( const TSharedPtr<FJsonValue >& Value, const FString& Identifier, const TSharedRef< TJsonWriter< CharType, PrintPolicy > >& Writer )
+	{
+		TSharedRef< FElement > StartingElement = MakeShareable( new FElement( Identifier, Value ) );
+		return FJsonSerializer::Serialize<CharType, PrintPolicy>( StartingElement, Writer );
+	}
+
+private:
+
+	struct StackState
+	{
+		EJson Type;
+		FString Identifier;
+		TArray<TSharedPtr<FJsonValue>> Array;
+		TSharedPtr<FJsonObject> Object;
+	};
+
+	struct FElement
+	{
+		FElement( const TSharedPtr<FJsonValue>& InValue )
+			: Identifier()
+			, Value(InValue)
+			, HasBeenProcessed(false)
+		{ }
+
+		FElement( const TSharedRef<FJsonObject>& Object )
+			: Identifier()
+			, Value(MakeShareable(new FJsonValueObject(Object)))
+			, HasBeenProcessed( false )
+		{ }
+
+		FElement( const TArray<TSharedPtr<FJsonValue>>& Array )
+			: Identifier()
+			, Value( MakeShareable(new FJsonValueArray(Array)))
+			, HasBeenProcessed(false)
+		{ }
+
+		FElement( const FString& InIdentifier, const TSharedPtr< FJsonValue >& InValue )
+			: Identifier( InIdentifier )
+			, Value( InValue )
+			, HasBeenProcessed( false )
+		{
+
+		}
+
+		FString Identifier;
+		TSharedPtr< FJsonValue > Value;
+		bool HasBeenProcessed;
+	};
+
+private:
+
+	template <class CharType>
+	static bool Deserialize( const TSharedRef< TJsonReader<CharType> >& Reader, StackState& OutStackState )
+	{
+		TArray<TSharedRef<StackState>> ScopeStack; 
+		TSharedPtr<StackState> CurrentState;
+
+		TSharedPtr<FJsonValue> NewValue;
+		EJsonNotation Notation;
+
+		while (Reader->ReadNext(Notation))
+		{
+			FString Identifier = Reader->GetIdentifier();
+			NewValue.Reset();
+
+			switch( Notation )
+			{
+			case EJsonNotation::ObjectStart:
+				{
+					if (CurrentState.IsValid())
+					{
+						ScopeStack.Push(CurrentState.ToSharedRef());
+					}
+
+					CurrentState = MakeShareable(new StackState());
+					CurrentState->Type = EJson::Object;
+					CurrentState->Identifier = Identifier;
+					CurrentState->Object = MakeShareable(new FJsonObject());
+				}
+				break;
+
+			case EJsonNotation::ObjectEnd:
+				{
+					if (ScopeStack.Num() > 0)
+					{
+						Identifier = CurrentState->Identifier;
+						NewValue = MakeShareable(new FJsonValueObject(CurrentState->Object));
+						CurrentState = ScopeStack.Pop();
+					}
+				}
+				break;
+
+			case EJsonNotation::ArrayStart:
+				{
+					if (CurrentState.IsValid())
+					{
+						ScopeStack.Push(CurrentState.ToSharedRef());
+					}
+
+					CurrentState = MakeShareable(new StackState());
+					CurrentState->Type = EJson::Array;
+					CurrentState->Identifier = Identifier;
+				}
+				break;
+
+			case EJsonNotation::ArrayEnd:
+				{
+					if (ScopeStack.Num() > 0)
+					{
+						Identifier = CurrentState->Identifier;
+						NewValue = MakeShareable(new FJsonValueArray(CurrentState->Array));
+						CurrentState = ScopeStack.Pop();
+					}
+				}
+				break;
+
+			case EJsonNotation::Boolean:
+				NewValue = MakeShareable(new FJsonValueBoolean(Reader->GetValueAsBoolean()));
+				break;
+
+			case EJsonNotation::String:
+				NewValue = MakeShareable(new FJsonValueString(Reader->GetValueAsString()));
+				break;
+
+			case EJsonNotation::Number:
+				NewValue = MakeShareable(new FJsonValueNumber(Reader->GetValueAsNumber()));
+				break;
+
+			case EJsonNotation::Null:
+				NewValue = MakeShareable(new FJsonValueNull());
+				break;
+
+			case EJsonNotation::Error:
+				return false;
+				break;
+			}
+
+			if (NewValue.IsValid() && CurrentState.IsValid())
+			{
+				if (CurrentState->Type == EJson::Object)
+				{
+					CurrentState->Object->SetField(Identifier, NewValue);
+				}
+				else
+				{
+					CurrentState->Array.Add(NewValue);
+				}
+			}
+		}
+
+		if (!CurrentState.IsValid() || !Reader->GetErrorMessage().IsEmpty())
+		{
+			return false;
+		}
+
+		OutStackState = *CurrentState.Get();
+
+		return true;
+	}
+
+	template <class CharType, class PrintPolicy>
+	static bool Serialize( const TSharedRef<FElement>& StartingElement, const TSharedRef<TJsonWriter<CharType, PrintPolicy>>& Writer )
+	{
+		TArray<TSharedRef<FElement>> ElementStack;
+		ElementStack.Push(StartingElement);
+
+		while (ElementStack.Num() > 0)
+		{
+			TSharedRef<FElement> Element = ElementStack.Pop();
+			check(Element->Value->Type != EJson::None);
+
+			switch (Element->Value->Type)
+			{
+			case EJson::Number:	
+				{
+					if (Element->Identifier.IsEmpty())
+					{
+						Writer->WriteValue(Element->Value->AsNumber());
+					}
+					else
+					{
+						Writer->WriteValue(Element->Identifier, Element->Value->AsNumber());
+					}
+				}
+				break;
+
+			case EJson::Boolean:					
+				{
+					if (Element->Identifier.IsEmpty())
+					{
+						Writer->WriteValue(Element->Value->AsBool());
+					}
+					else
+					{
+						Writer->WriteValue(Element->Identifier, Element->Value->AsBool());
+					}
+				}
+				break;
+
+			case EJson::String:
+				{
+					if (Element->Identifier.IsEmpty())
+					{
+						Writer->WriteValue(Element->Value->AsString());
+					}
+					else
+					{
+						Writer->WriteValue(Element->Identifier, Element->Value->AsString());
+					}
+				}
+				break;
+
+			case EJson::Null:
+				{
+					if (Element->Identifier.IsEmpty())
+					{
+						Writer->WriteNull();	
+					}
+					else
+					{
+						Writer->WriteNull(Element->Identifier);
+					}
+				}
+				break;
+
+			case EJson::Array:
+				{
+					if (Element->HasBeenProcessed)
+					{
+						Writer->WriteArrayEnd();
+					}
+					else
+					{
+						Element->HasBeenProcessed = true;
+						ElementStack.Push(Element);
+
+						if (Element->Identifier.IsEmpty())
+						{
+							Writer->WriteArrayStart();
+						}
+						else
+						{
+							Writer->WriteArrayStart(Element->Identifier);
+						}
+
+						TArray<TSharedPtr<FJsonValue>> Values = Element->Value->AsArray();
+
+						for (int Index = Values.Num() - 1; Index >= 0; --Index)
+						{
+							ElementStack.Push(MakeShareable(new FElement(Values[Index])));
+						}
+					}
+				}
+				break;
+
+			case EJson::Object:
+				{
+					if (Element->HasBeenProcessed)
+					{
+						Writer->WriteObjectEnd();
+					}
+					else
+					{
+						Element->HasBeenProcessed = true;
+						ElementStack.Push(Element);
+
+						if (Element->Identifier.IsEmpty())
+						{
+							Writer->WriteObjectStart();
+						}
+						else
+						{
+							Writer->WriteObjectStart(Element->Identifier);
+						}
+
+						TArray<FString> Keys; 
+						TArray<TSharedPtr<FJsonValue>> Values;
+						TSharedPtr<FJsonObject> ElementObject = Element->Value->AsObject();
+						ElementObject->Values.GenerateKeyArray(Keys);
+						ElementObject->Values.GenerateValueArray(Values);
+
+						check(Keys.Num() == Values.Num());
+
+						for (int Index = Values.Num() - 1; Index >= 0; --Index)
+						{
+							ElementStack.Push(MakeShareable(new FElement(Keys[Index], Values[Index ])));
+						}
+					}
+				}
+				break;
+
+			default: 
+				UE_LOG(LogJson, Fatal,TEXT("Could not print Json Value, unrecognized type."));
+			}
+		}
+
+		return Writer->Close();
+	}
+};
diff --git a/Engine/Source/Runtime/Json/Public/Serialization/JsonTypes.h b/Engine/Source/Runtime/Json/Public/Serialization/JsonTypes.h
new file mode 100644
--- /dev/null
+++ b/Engine/Source/Runtime/Json/Public/Serialization/JsonTypes.h
@@ -0,0 +1,59 @@
+// Copyright 1998-2014 Epic Games, Inc. All Rights Reserved.
+
+#pragma once
+
+
+/**
+ * Json (JavaScript Object Notation) is a lightweight data-interchange format.
+ * Information on how it works can be found here: http://www.json.org/.
+ * This code was written from scratch with only the Json spec as a guide.
+ *
+ * In order to use Json effectively, you need to be familiar with the Object/Value
+ * hierarchy, and you should use the FJsonObject class and FJsonValue subclasses.
+ */
+
+/**
+ * Represents all the types a Json Value can be.
+ */
+enum class EJson
+{
+	None,
+	Null,
+	String,
+	Number,
+	Boolean,
+	Array,
+	Object
+};
+
+
+enum class EJsonToken
+{
+	None,
+	Comma,
+	CurlyOpen,
+	CurlyClose,
+	SquareOpen,
+	SquareClose,
+	Colon,
+	String,
+	Number,
+	True,
+	False,
+	Null,
+	Identifier
+};
+
+
+enum class EJsonNotation
+{
+	ObjectStart,
+	ObjectEnd,
+	ArrayStart,
+	ArrayEnd,
+	Boolean,
+	String,
+	Number,
+	Null,
+	Error
+};
diff --git a/Engine/Source/Runtime/Json/Public/Serialization/JsonWriter.h b/Engine/Source/Runtime/Json/Public/Serialization/JsonWriter.h
new file mode 100644
--- /dev/null
+++ b/Engine/Source/Runtime/Json/Public/Serialization/JsonWriter.h
@@ -0,0 +1,445 @@
+// Copyright 1998-2014 Epic Games, Inc. All Rights Reserved.
+
+#pragma once
+
+
+/**
+ * Template for Json writers.
+ *
+ * @param CharType The type of characters to print, i.e. TCHAR or ANSICHAR.
+ * @param PrintPolicy The print policy to use when writing the output string (default = TPrettyJsonPrintPolicy).
+ */
+template <class CharType = TCHAR, class PrintPolicy = TPrettyJsonPrintPolicy<CharType> >
+class TJsonWriter
+{
+public:
+
+	static TSharedRef< TJsonWriter > Create( FArchive* const Stream, int32 InitialIndentLevel = 0 )
+	{
+		return MakeShareable( new TJsonWriter< CharType, PrintPolicy >( Stream, InitialIndentLevel ) );
+	}
+
+public:
+
+	virtual ~TJsonWriter() { }
+
+	FORCEINLINE int32 GetIndentLevel() const { return IndentLevel; }
+
+	void WriteObjectStart()
+	{
+		check(CanWriteValueWithoutIdentifier());
+		if (PreviousTokenWritten != EJsonToken::None )
+		{
+			WriteCommaIfNeeded();
+		}
+
+		if ( PreviousTokenWritten != EJsonToken::None )
+		{
+			PrintPolicy::WriteLineTerminator(Stream);
+			PrintPolicy::WriteTabs(Stream, IndentLevel);
+		}
+
+		PrintPolicy::WriteChar(Stream, CharType('{'));
+		++IndentLevel;
+		Stack.Push( EJson::Object );
+		PreviousTokenWritten = EJsonToken::CurlyOpen;
+	}
+
+	void WriteObjectStart( const FString& Identifier )
+	{
+		check( Stack.Top() == EJson::Object );
+		WriteIdentifier( Identifier );
+
+		PrintPolicy::WriteLineTerminator(Stream);
+		PrintPolicy::WriteTabs(Stream, IndentLevel);
+		PrintPolicy::WriteChar(Stream, CharType('{'));
+		++IndentLevel;
+		Stack.Push( EJson::Object );
+		PreviousTokenWritten = EJsonToken::CurlyOpen;
+	}
+
+	void WriteObjectEnd()
+	{
+		check( Stack.Top() == EJson::Object );
+
+		PrintPolicy::WriteLineTerminator(Stream);
+
+		--IndentLevel;
+		PrintPolicy::WriteTabs(Stream, IndentLevel);
+		PrintPolicy::WriteChar(Stream, CharType('}'));
+		Stack.Pop();
+		PreviousTokenWritten = EJsonToken::CurlyClose;
+	}
+
+	void WriteArrayStart()
+	{
+		check(CanWriteValueWithoutIdentifier());
+		if ( PreviousTokenWritten != EJsonToken::None )
+		{
+			WriteCommaIfNeeded();
+		}
+
+		if ( PreviousTokenWritten != EJsonToken::None )
+		{
+			PrintPolicy::WriteLineTerminator(Stream);
+			PrintPolicy::WriteTabs(Stream, IndentLevel);
+		}
+
+		PrintPolicy::WriteChar(Stream, CharType('['));
+		++IndentLevel;
+		Stack.Push( EJson::Array );
+		PreviousTokenWritten = EJsonToken::SquareOpen;
+	}
+
+	void WriteArrayStart( const FString& Identifier )
+	{
+		check( Stack.Top() == EJson::Object );
+		WriteIdentifier( Identifier );
+
+		PrintPolicy::WriteSpace( Stream );
+		PrintPolicy::WriteChar(Stream, CharType('['));
+		++IndentLevel;
+		Stack.Push( EJson::Array );
+		PreviousTokenWritten = EJsonToken::SquareOpen;
+	}
+
+	void WriteArrayEnd()
+	{
+		check( Stack.Top() == EJson::Array );
+
+		--IndentLevel;
+		if ( PreviousTokenWritten == EJsonToken::SquareClose || PreviousTokenWritten == EJsonToken::CurlyClose || PreviousTokenWritten == EJsonToken::String )
+		{
+			PrintPolicy::WriteLineTerminator(Stream);
+			PrintPolicy::WriteTabs(Stream, IndentLevel);
+		}
+		else if ( PreviousTokenWritten != EJsonToken::SquareOpen )
+		{
+			PrintPolicy::WriteSpace( Stream );
+		}
+
+		PrintPolicy::WriteChar(Stream, CharType(']'));
+		Stack.Pop();
+		PreviousTokenWritten = EJsonToken::SquareClose;
+	}
+
+	void WriteValue( const FString& Identifier, const bool Value )
+	{
+		check( Stack.Top() == EJson::Object );
+		WriteIdentifier( Identifier );
+
+		PrintPolicy::WriteSpace(Stream);
+		WriteBoolValue( Value );
+		PreviousTokenWritten = Value ? EJsonToken::True : EJsonToken::False;
+	}
+
+	void WriteValue( const FString& Identifier, const double Value )
+	{
+		check( Stack.Top() == EJson::Object );
+		WriteIdentifier( Identifier );
+
+		PrintPolicy::WriteSpace(Stream);
+		WriteNumberValue( Value );
+		PreviousTokenWritten = EJsonToken::Number;
+	}
+
+	void WriteValue( const FString& Identifier, const int32 Value )
+	{
+		check( Stack.Top() == EJson::Object );
+		WriteIdentifier( Identifier );
+
+		PrintPolicy::WriteSpace(Stream);
+		WriteIntegerValue( Value );
+		PreviousTokenWritten = EJsonToken::Number;
+	}
+
+	void WriteValue( const FString& Identifier, const int64 Value )
+	{
+		check( Stack.Top() == EJson::Object );
+		WriteIdentifier( Identifier );
+
+		PrintPolicy::WriteSpace(Stream);
+		WriteIntegerValue( Value );
+		PreviousTokenWritten = EJsonToken::Number;
+	}
+
+	void WriteValue( const FString& Identifier, const FString& Value )
+	{
+		check( Stack.Top() == EJson::Object );
+		WriteIdentifier( Identifier );
+
+		PrintPolicy::WriteSpace(Stream);
+		WriteStringValue( Value );
+		PreviousTokenWritten = EJsonToken::String;
+	}
+
+	void WriteValue( const FString& Identifier, const TCHAR* Value )
+	{
+		WriteValue(Identifier, FString(Value));
+	}
+
+	// WARNING: THIS IS DANGEROUS. Use this only if you know for a fact that the Value is valid JSON!
+	// Use this to insert the results of a different JSON Writer in.
+	void WriteRawJSONValue( const FString& Identifier, const FString& Value )
+	{
+		check( Stack.Top() == EJson::Object );
+		WriteIdentifier( Identifier );
+
+		PrintPolicy::WriteSpace(Stream);
+		PrintPolicy::WriteString(Stream, Value);
+		PreviousTokenWritten = EJsonToken::String;
+	}
+
+	void WriteNull( const FString& Identifier )
+	{
+		check( Stack.Top() == EJson::Object );
+		WriteIdentifier( Identifier );
+
+		PrintPolicy::WriteSpace(Stream);
+		WriteNullValue();
+		PreviousTokenWritten = EJsonToken::Null;
+	}
+
+	void WriteValue( const bool Value )
+	{
+		check(CanWriteValueWithoutIdentifier());
+		WriteCommaIfNeeded();
+
+		if ( PreviousTokenWritten != EJsonToken::True && PreviousTokenWritten != EJsonToken::False && PreviousTokenWritten != EJsonToken::SquareOpen )
+		{
+			PrintPolicy::WriteLineTerminator(Stream);
+			PrintPolicy::WriteTabs(Stream, IndentLevel);
+		}
+		else
+		{
+			PrintPolicy::WriteSpace( Stream );
+		}
+
+		WriteBoolValue( Value );
+		PreviousTokenWritten = Value ? EJsonToken::True : EJsonToken::False;
+	}
+
+	void WriteValue( const double Value )
+	{
+		check(CanWriteValueWithoutIdentifier());
+		WriteCommaIfNeeded();
+
+		if ( PreviousTokenWritten != EJsonToken::Number && PreviousTokenWritten != EJsonToken::SquareOpen )
+		{
+			PrintPolicy::WriteLineTerminator(Stream);
+			PrintPolicy::WriteTabs(Stream, IndentLevel);
+		}
+		else
+		{
+			PrintPolicy::WriteSpace( Stream );
+		}
+
+		WriteNumberValue( Value );
+		PreviousTokenWritten = EJsonToken::Number;
+	}
+
+	void WriteValue( const FString& Value )
+	{
+		check(CanWriteValueWithoutIdentifier());
+		WriteCommaIfNeeded();
+		PrintPolicy::WriteLineTerminator(Stream);
+		PrintPolicy::WriteTabs(Stream, IndentLevel);
+		WriteStringValue( Value );
+		PreviousTokenWritten = EJsonToken::String;
+	}
+
+	void WriteNull()
+	{
+		check(CanWriteValueWithoutIdentifier());
+		WriteCommaIfNeeded();
+
+		if ( PreviousTokenWritten != EJsonToken::Null && PreviousTokenWritten != EJsonToken::SquareOpen )
+		{
+			PrintPolicy::WriteLineTerminator(Stream);
+			PrintPolicy::WriteTabs(Stream, IndentLevel);
+		}
+		else
+		{
+			PrintPolicy::WriteSpace( Stream );
+		}
+
+		WriteNullValue();
+		PreviousTokenWritten = EJsonToken::Null;
+	}
+
+	virtual bool Close()
+	{
+		return ( PreviousTokenWritten == EJsonToken::None || 
+				 PreviousTokenWritten == EJsonToken::CurlyClose  || 
+				 PreviousTokenWritten == EJsonToken::SquareClose ) 
+				&& Stack.Num() == 0;
+	}
+
+	/**
+	 * WriteValue("Foo", Bar) should be equivalent to WriteIdentifierPrefix("Foo"), WriteValue(Bar)
+	 */
+	void WriteIdentifierPrefix(const FString& Identifier)
+	{
+		check(Stack.Top() == EJson::Object);
+		WriteIdentifier(Identifier);
+		PrintPolicy::WriteSpace(Stream);
+		PreviousTokenWritten = EJsonToken::Identifier;
+	}
+
+protected:
+
+	/**
+	 * Creates and initializes a new instance.
+	 *
+	 * @param InStream An archive containing the input.
+	 * @param InitialIndentLevel The initial indentation level.
+	 */
+	TJsonWriter( FArchive* const InStream, int32 InitialIndentLevel )
+		: Stream( InStream )
+		, Stack()
+		, PreviousTokenWritten(EJsonToken::None)
+		, IndentLevel(InitialIndentLevel)
+	{ }
+
+protected:
+
+	FORCEINLINE bool CanWriteValueWithoutIdentifier() const 
+	{ 
+		return Stack.Num() <= 0 || Stack.Top() == EJson::Array || PreviousTokenWritten == EJsonToken::Identifier;
+	}
+
+	FORCEINLINE void WriteCommaIfNeeded()
+	{
+		if ( PreviousTokenWritten != EJsonToken::CurlyOpen && PreviousTokenWritten != EJsonToken::SquareOpen && PreviousTokenWritten != EJsonToken::Identifier)
+		{
+			PrintPolicy::WriteChar(Stream, CharType(','));
+		}
+	}
+
+	FORCEINLINE void WriteIdentifier( const FString& Identifier )
+	{
+		WriteCommaIfNeeded();
+		PrintPolicy::WriteLineTerminator(Stream);
+
+		PrintPolicy::WriteTabs(Stream, IndentLevel);
+		WriteStringValue( Identifier );
+		PrintPolicy::WriteChar(Stream, CharType(':'));
+	}
+
+	FORCEINLINE void WriteBoolValue( const bool Value )
+	{
+		PrintPolicy::WriteString(Stream, Value ? TEXT("true") : TEXT("false"));
+	}
+
+	FORCEINLINE void WriteNumberValue( const double Value)
+	{
+		// Specify 17 significant digits, the most that can ever be useful from a double
+		// In particular, this ensures large integers are written correctly
+		PrintPolicy::WriteString(Stream, FString::Printf(TEXT("%.17g"), Value));
+	}
+
+	FORCEINLINE void WriteIntegerValue( const int64 Value)
+	{
+		PrintPolicy::WriteString(Stream, FString::Printf(TEXT("%lld"), Value));
+	}
+
+	FORCEINLINE void WriteNullValue()
+	{
+		PrintPolicy::WriteString(Stream, TEXT("null"));
+	}
+
+	virtual void WriteStringValue( const FString& String )
+	{
+		FString OutString;
+		OutString += TEXT("\"");
+		for (const TCHAR* Char = *String; *Char != TCHAR('\0'); ++Char)
+		{
+			switch (*Char)
+			{
+			case TCHAR('\\'): OutString += TEXT("\\\\"); break;
+			case TCHAR('\n'): OutString += TEXT("\\n"); break;
+			case TCHAR('\t'): OutString += TEXT("\\t"); break;
+			case TCHAR('\b'): OutString += TEXT("\\b"); break;
+			case TCHAR('\f'): OutString += TEXT("\\f"); break;
+			case TCHAR('\r'): OutString += TEXT("\\r"); break;
+			case TCHAR('\"'): OutString += TEXT("\\\""); break;
+			default: OutString += *Char;
+			}
+		}
+		OutString += TEXT("\"");
+		PrintPolicy::WriteString(Stream, OutString);
+	}
+
+protected:
+
+	FArchive* const Stream;
+	TArray<EJson> Stack;
+	EJsonToken PreviousTokenWritten;
+	int32 IndentLevel;
+};
+
+
+template <class PrintPolicy = TPrettyJsonPrintPolicy<TCHAR>>
+class TJsonStringWriter
+	: public TJsonWriter<TCHAR, PrintPolicy>
+{
+public:
+
+	static TSharedRef<TJsonStringWriter> Create( FString* const InStream, int32 InitialIndent = 0 )
+	{
+		return MakeShareable(new TJsonStringWriter(InStream, InitialIndent));
+	}
+
+public:
+
+	virtual ~TJsonStringWriter()
+	{
+		check(this->Stream->Close());
+		delete this->Stream;
+	}
+
+	virtual bool Close() override
+	{
+		FString Out;
+
+		for (int32 i = 0; i < Bytes.Num(); i+=sizeof(TCHAR))
+		{
+			TCHAR* Char = static_cast<TCHAR*>(static_cast<void*>(&Bytes[i]));
+			Out += *Char;
+		}
+
+		*OutString = Out;
+
+		return TJsonWriter<TCHAR, PrintPolicy>::Close();
+	}
+
+protected:
+
+	TJsonStringWriter( FString* const InOutString, int32 InitialIndent )
+		: TJsonWriter<TCHAR, PrintPolicy>(new FMemoryWriter(Bytes), InitialIndent)
+		, Bytes()
+		, OutString(InOutString)
+	{ }
+
+private:
+
+	TArray<uint8> Bytes;
+	FString* OutString;
+};
+
+
+template <class CharType = TCHAR, class PrintPolicy = TPrettyJsonPrintPolicy<CharType>>
+class TJsonWriterFactory
+{
+public:
+
+	static TSharedRef<TJsonWriter<CharType, PrintPolicy>> Create( FArchive* const Stream, int32 InitialIndent = 0 )
+	{
+		return TJsonWriter< CharType, PrintPolicy >::Create(Stream, InitialIndent);
+	}
+
+	static TSharedRef<TJsonWriter<TCHAR, PrintPolicy>> Create( FString* const Stream, int32 InitialIndent = 0 )
+	{
+		return StaticCastSharedRef<TJsonWriter<TCHAR, PrintPolicy>>(TJsonStringWriter<PrintPolicy>::Create(Stream, InitialIndent));
+	}
+};
diff --git a/Engine/Source/Runtime/JsonUtilities/JsonUtilities.Build.cs b/Engine/Source/Runtime/JsonUtilities/JsonUtilities.Build.cs
--- a/Engine/Source/Runtime/JsonUtilities/JsonUtilities.Build.cs
+++ b/Engine/Source/Runtime/JsonUtilities/JsonUtilities.Build.cs
@@ -6,12 +6,18 @@ public class JsonUtilities : ModuleRules
 {
 	public JsonUtilities( TargetInfo Target )
 	{
+		PublicDependencyModuleNames.AddRange(
+			new string[]
+			{
+				"Core",
+				"Json",
+			}
+		);
 		PrivateDependencyModuleNames.AddRange(
 			new string[] 
-			{ 
-				"Core",
+			{ 			
 				"CoreUObject",
 			}
-			);
+		);
 	}
 }
diff --git a/Engine/Source/Runtime/Messaging/Messaging.Build.cs b/Engine/Source/Runtime/Messaging/Messaging.Build.cs
--- a/Engine/Source/Runtime/Messaging/Messaging.Build.cs
+++ b/Engine/Source/Runtime/Messaging/Messaging.Build.cs
@@ -17,6 +17,7 @@ namespace UnrealBuildTool.Rules
 				new string[]
 				{
 					"Core",
+					"Json",
 					"Sockets",
 				}
 			);
diff --git a/Engine/Source/Runtime/Messaging/Private/Serialization/JsonMessageSerializer.cpp b/Engine/Source/Runtime/Messaging/Private/Serialization/JsonMessageSerializer.cpp
--- a/Engine/Source/Runtime/Messaging/Private/Serialization/JsonMessageSerializer.cpp
+++ b/Engine/Source/Runtime/Messaging/Private/Serialization/JsonMessageSerializer.cpp
@@ -75,7 +75,8 @@ namespace Obsolete
 	UProperty* FindProperty( const FJsonStructReadState& State, const FString& Identifier )
 	{
 		UArrayProperty* ArrayProperty = Cast<UArrayProperty>( State.Property );
-		if ( ArrayProperty )
+
+		if (ArrayProperty != nullptr)
 		{
 			return ArrayProperty->Inner;
 		}
@@ -126,12 +127,12 @@ namespace Obsolete
 			FScriptArrayHelper ArrayHelper(ArrayProperty, ArrayProperty->template ContainerPtrToValuePtr<void>(State.Data));
 			int32 Index = ArrayHelper.AddValue();
 		
-			ValuePtr = (PropertyType*)ArrayHelper.GetRawPtr( Index );
+			ValuePtr = (PropertyType*)ArrayHelper.GetRawPtr(Index);
 		}
 		else
 		{
-			UPropertyType* TypedProperty = Cast<UPropertyType>( Property );
-			check( TypedProperty )
+			UPropertyType* TypedProperty = Cast<UPropertyType>(Property);
+			check(TypedProperty)
 
 			ValuePtr = TypedProperty->template ContainerPtrToValuePtr<PropertyType>(State.Data);
 		}
@@ -149,18 +150,19 @@ namespace Obsolete
 	}
 
 
-	bool ReadExpectedArrayValue( const TSharedRef< TJsonReader<UCS2CHAR> >& Reader, const EJsonNotation::Type ExpectNotation, const FString& ErrorMessage )
+	bool ReadExpectedArrayValue( const TSharedRef<TJsonReader<UCS2CHAR>>& Reader, const EJsonNotation ExpectNotation, const FString& ErrorMessage )
 	{
 		bool ReadCorrectValue = true;
 
-		EJsonNotation::Type Notation;
-		if ( !Reader->ReadNext( Notation ) || Notation != ExpectNotation )
+		EJsonNotation Notation;
+
+		if (!Reader->ReadNext(Notation) || (Notation != ExpectNotation))
 		{
 			ReadCorrectValue = false;
 
-			UE_LOG( LogJson, Log, TEXT("%s Line: %u Ch: %u"), *ErrorMessage, Reader->GetLineNumber(), Reader->GetCharacterNumber() );
+			UE_LOG(LogJson, Log, TEXT("%s Line: %u Ch: %u"), *ErrorMessage, Reader->GetLineNumber(), Reader->GetCharacterNumber());
 		
-			if ( Notation != EJsonNotation::ArrayEnd )
+			if (Notation != EJsonNotation::ArrayEnd)
 			{
 				Reader->SkipArray();
 			}
@@ -242,7 +244,7 @@ bool FJsonMessageSerializer::DeserializeStruct( IMutableMessageContextRef& OutCo
 	CurrentState.Property = nullptr;
 	CurrentState.TypeInfo = nullptr;
 
-	EJsonNotation::Type Notation;
+	EJsonNotation Notation;
 
 	while (Reader->ReadNext(Notation))
 	{
diff --git a/Engine/Source/Runtime/Online/BuildPatchServices/BuildPatchServices.Build.cs b/Engine/Source/Runtime/Online/BuildPatchServices/BuildPatchServices.Build.cs
--- a/Engine/Source/Runtime/Online/BuildPatchServices/BuildPatchServices.Build.cs
+++ b/Engine/Source/Runtime/Online/BuildPatchServices/BuildPatchServices.Build.cs
@@ -22,9 +22,10 @@ public class BuildPatchServices : ModuleRules
 
 		PrivateDependencyModuleNames.AddRange(
 		new string[] {
-				"HTTP",
 				"Analytics",
 				"AnalyticsET",
+				"HTTP",
+				"Json",
 			}
 		);
 
diff --git a/Engine/Source/Runtime/Online/OnlineSubsystem/OnlineSubsystem.Build.cs b/Engine/Source/Runtime/Online/OnlineSubsystem/OnlineSubsystem.Build.cs
--- a/Engine/Source/Runtime/Online/OnlineSubsystem/OnlineSubsystem.Build.cs
+++ b/Engine/Source/Runtime/Online/OnlineSubsystem/OnlineSubsystem.Build.cs
@@ -6,6 +6,12 @@ public class OnlineSubsystem : ModuleRules
 {
 	public OnlineSubsystem(TargetInfo Target)
 	{
+		PublicDependencyModuleNames.AddRange(
+			new string[] {
+				"Json",
+			}
+		);
+
         PublicIncludePaths.Add("Runtime/Online/OnlineSubsystem/Public/Interfaces");
 
 		PrivateIncludePaths.Add("Runtime/Online/OnlineSubsystem/Private");		
diff --git a/Engine/Source/Runtime/Online/OnlineSubsystemAmazon/OnlineSubsystemAmazon.Build.cs b/Engine/Source/Runtime/Online/OnlineSubsystemAmazon/OnlineSubsystemAmazon.Build.cs
--- a/Engine/Source/Runtime/Online/OnlineSubsystemAmazon/OnlineSubsystemAmazon.Build.cs
+++ b/Engine/Source/Runtime/Online/OnlineSubsystemAmazon/OnlineSubsystemAmazon.Build.cs
@@ -13,8 +13,9 @@ public class OnlineSubsystemAmazon : ModuleRules
 			new string[] {
 				"Core", 
 				"HTTP",
+				"Json",
                 "OnlineSubsystem", 
 			}
-			);
+		);
 	}
 }
diff --git a/Engine/Source/Runtime/Online/OnlineSubsystemFacebook/OnlineSubsystemFacebook.Build.cs b/Engine/Source/Runtime/Online/OnlineSubsystemFacebook/OnlineSubsystemFacebook.Build.cs
--- a/Engine/Source/Runtime/Online/OnlineSubsystemFacebook/OnlineSubsystemFacebook.Build.cs
+++ b/Engine/Source/Runtime/Online/OnlineSubsystemFacebook/OnlineSubsystemFacebook.Build.cs
@@ -14,8 +14,9 @@ public class OnlineSubsystemFacebook : ModuleRules
             new string[] { 
                 "Core", 
                 "HTTP",
+				"ImageCore",
+				"Json",
                 "OnlineSubsystem", 
-                "ImageCore",
             }
             );
 
diff --git a/Engine/Source/Runtime/Projects/Private/ProjectsPrivatePCH.h b/Engine/Source/Runtime/Projects/Private/ProjectsPrivatePCH.h
--- a/Engine/Source/Runtime/Projects/Private/ProjectsPrivatePCH.h
+++ b/Engine/Source/Runtime/Projects/Private/ProjectsPrivatePCH.h
@@ -1,9 +1,5 @@
 // Copyright 1998-2014 Epic Games, Inc. All Rights Reserved.
 
-/*=============================================================================
-	ProjectsPrivatePCH.h: Pre-compiled header file for the Projects module.
-=============================================================================*/
-
 #pragma once
 
 
diff --git a/Engine/Source/Runtime/Projects/Projects.Build.cs b/Engine/Source/Runtime/Projects/Projects.Build.cs
--- a/Engine/Source/Runtime/Projects/Projects.Build.cs
+++ b/Engine/Source/Runtime/Projects/Projects.Build.cs
@@ -10,6 +10,7 @@ namespace UnrealBuildTool.Rules
 				new string[]
 				{
 					"Core",
+					"Json",
 				}
 			);
 
diff --git a/Engine/Source/Runtime/Slate/Slate.Build.cs b/Engine/Source/Runtime/Slate/Slate.Build.cs
--- a/Engine/Source/Runtime/Slate/Slate.Build.cs
+++ b/Engine/Source/Runtime/Slate/Slate.Build.cs
@@ -13,6 +13,7 @@ public class Slate : ModuleRules
 				"Core",
 				"CoreUObject",
 				"InputCore",
+				"Json",
 				"SlateCore",
 			}
 		);
diff --git a/Engine/Source/Runtime/SlateCore/SlateCore.Build.cs b/Engine/Source/Runtime/SlateCore/SlateCore.Build.cs
--- a/Engine/Source/Runtime/SlateCore/SlateCore.Build.cs
+++ b/Engine/Source/Runtime/SlateCore/SlateCore.Build.cs
@@ -11,6 +11,7 @@ public class SlateCore : ModuleRules
 				"Core",
 				"CoreUObject",
 				"InputCore",
+				"Json",
 			}
 		);
 
