fix nonunity build

diff --git a/Engine/Source/Runtime/Internationalization/Private/Serialization/JsonInternationalizationMetadataSerializer.cpp b/Engine/Source/Runtime/Internationalization/Private/Serialization/JsonInternationalizationMetadataSerializer.cpp
--- a/Engine/Source/Runtime/Internationalization/Private/Serialization/JsonInternationalizationMetadataSerializer.cpp
+++ b/Engine/Source/Runtime/Internationalization/Private/Serialization/JsonInternationalizationMetadataSerializer.cpp
@@ -3,6 +3,7 @@
 #include "InternationalizationPrivatePCH.h"
 #include "Internationalization/InternationalizationMetadata.h"
 #include "Json.h"
+#include "JsonInternationalizationMetadataSerializer.h"
 
 
 struct FCompareLocMetadataValue
