Original changes by [REDACTED_PERSON]. Unshelved and made to compile. Resolves issue where locres generation improperly load archives without checking the UTF16 BOM.

diff --git a/Engine/Source/Runtime/Core/Private/Internationalization/TextLocalizationManager.cpp b/Engine/Source/Runtime/Core/Private/Internationalization/TextLocalizationManager.cpp
--- a/Engine/Source/Runtime/Core/Private/Internationalization/TextLocalizationManager.cpp
+++ b/Engine/Source/Runtime/Core/Private/Internationalization/TextLocalizationManager.cpp
@@ -715,12 +715,14 @@ void FTextLocalizationManager::RegenerateResources( const FString& ConfigFilePat
 	for(int32 i = 0; i < BackingBuffers.Num(); ++i)
 	{
 		TArray<uint8>& BackingBuffer = BackingBuffers[i];
-
 		FMemoryWriter MemoryWriter(BackingBuffer, true);
 
 		// Read the manifest file from the source path.
 		FString ManifestFilePath = (SourcePath / ManifestName);
 		ManifestFilePath = FPaths::ConvertRelativePathToFull(ManifestFilePath);
+		TSharedRef<FInternationalizationManifest> InternationalizationManifest = MakeShareable(new FInternationalizationManifest);
+
+#if 0 // @todo Json: Serializing from FArchive is currently broken
 		FArchive* ManifestFile = IFileManager::Get().CreateFileReader(*ManifestFilePath);
 
 		if (ManifestFile == nullptr)
@@ -729,11 +731,19 @@ void FTextLocalizationManager::RegenerateResources( const FString& ConfigFilePat
 			return;
 		}
 
-		TSharedRef<FInternationalizationManifest> InternationalizationManifest = MakeShareable( new FInternationalizationManifest );
+		ManifestSerializer.DeserializeManifest(*ManifestFile, InternationalizationManifest);
+#else
+		FString ManifestContent;
+
+		if (!FFileHelper::LoadFileToString(ManifestContent, *ManifestFilePath))
 		{
-			ManifestSerializer.DeserializeManifest(*ManifestFile, InternationalizationManifest);
+			UE_LOG(LogTextLocalizationManager, Error, TEXT("Failed to load file %s."), *ManifestFilePath);
+			continue;
 		}
 
+		ManifestSerializer.DeserializeManifest(ManifestContent, InternationalizationManifest);
+#endif
+
 		// Write resource.
 		FTextLocalizationResourceGenerator::Generate(SourcePath, InternationalizationManifest, LocaleNames[i], &(MemoryWriter), ArchiveSerializer);
 
@@ -745,11 +755,8 @@ void FTextLocalizationManager::RegenerateResources( const FString& ConfigFilePat
 	for(int32 i = 0; i < BackingBuffers.Num(); ++i)
 	{
 		TArray<uint8>& BackingBuffer = BackingBuffers[i];
-
 		FMemoryReader MemoryReader(BackingBuffer, true);
-
 		const FString CulturePath = DestinationPath / LocaleNames[i];
-
 		const FString ResourceFilePath = FPaths::ConvertRelativePathToFull(CulturePath / ResourceName);
 
 		FLocalizationEntryTracker& CultureTracker = LocalizationEntryTrackers[LocalizationEntryTrackers.Add(FLocalizationEntryTracker())];
diff --git a/Engine/Source/Runtime/Core/Private/Internationalization/TextLocalizationResourceGenerator.cpp b/Engine/Source/Runtime/Core/Private/Internationalization/TextLocalizationResourceGenerator.cpp
--- a/Engine/Source/Runtime/Core/Private/Internationalization/TextLocalizationResourceGenerator.cpp
+++ b/Engine/Source/Runtime/Core/Private/Internationalization/TextLocalizationResourceGenerator.cpp
@@ -164,6 +164,9 @@ bool FTextLocalizationResourceGenerator::Generate(const FString& SourcePath, con
 		// Read each archive file from the culture-named directory in the source path.
 		FString ArchiveFilePath = CulturePath / ArchiveName;
 		ArchiveFilePath = FPaths::ConvertRelativePathToFull(ArchiveFilePath);
+		TSharedRef<FInternationalizationArchive> InternationalizationArchive = MakeShareable(new FInternationalizationArchive);
+
+#if 0 // @todo Json: Serializing from FArchive is currently broken
 		FArchive* ArchiveFile = IFileManager::Get().CreateFileReader(*ArchiveFilePath);
 
 		if (ArchiveFile == nullptr)
@@ -171,12 +174,20 @@ bool FTextLocalizationResourceGenerator::Generate(const FString& SourcePath, con
 			UE_LOG(LogTextLocalizationResourceGenerator, Error, TEXT("No archive found at %s."), *ArchiveFilePath);
 			continue;
 		}
+			
+		ArchiveSerializer.DeserializeArchive(*ArchiveFile, InternationalizationArchive);
+#else
+		FString ArchiveContent;
 
-		TSharedRef<FInternationalizationArchive> InternationalizationArchive = MakeShareable( new FInternationalizationArchive );
+		if (!FFileHelper::LoadFileToString(ArchiveContent, *ArchiveFilePath))
 		{
-			ArchiveSerializer.DeserializeArchive(*ArchiveFile, InternationalizationArchive);
+			UE_LOG(LogTextLocalizationResourceGenerator, Error, TEXT("Failed to load file %s."), *ArchiveFilePath);
+			continue;
 		}
 
+		ArchiveSerializer.DeserializeArchive(ArchiveContent, InternationalizationArchive);
+#endif
+
 		// Generate text localization resource from manifest and archive entries.
 		for(TManifestEntryByContextIdContainer::TConstIterator i = InternationalizationManifest->GetEntriesByContextIdIterator(); i; ++i)
 		{
diff --git a/Engine/Source/Runtime/Core/Public/Internationalization/IInternationalizationArchiveSerializer.h b/Engine/Source/Runtime/Core/Public/Internationalization/IInternationalizationArchiveSerializer.h
--- a/Engine/Source/Runtime/Core/Public/Internationalization/IInternationalizationArchiveSerializer.h
+++ b/Engine/Source/Runtime/Core/Public/Internationalization/IInternationalizationArchiveSerializer.h
@@ -2,8 +2,10 @@
 
 #pragma once
 
+
 class FInternationalizationArchive;
 
+
 /**
  * Interface for Internationalization archive serializers.
  */
@@ -12,28 +14,45 @@ class IInternationalizationArchiveSerializer
 public:
 
 	/**
-	 * Virtual destructor.
+	 * Deserializes a Internationalization archive from a string.
+	 *
+	 * @param InStr The string to serialize from.
+	 * @param InternationalizationArchive The populated internationalization archive.
+	 * @return true if deserialization was successful, false otherwise.
 	 */
-	virtual ~IInternationalizationArchiveSerializer( ) { }
+	virtual bool DeserializeArchive( const FString& InStr, TSharedRef<FInternationalizationArchive> InternationalizationArchive ) = 0;
 
 	/**
-	 * Deserializes a Internationalizationarchive from an archive.
+	 * Serializes a Internationalization archive to a string.
 	 *
-	 * @param Archive - The archive to serialize from.
-	 * @param InternationalizationArchive - Will hold the content of the deserialized Internationalizationarchive.
+	 * @param InternationalizationArchive The Internationalization archive data to serialize.
+	 * @param Str The string to serialize into.
+	 * @return true if serialization was successful, false otherwise.
+	 */
+	virtual bool SerializeArchive( TSharedRef<const FInternationalizationArchive> InternationalizationArchive, FString& Str ) = 0;
+
+#if 0 // @todo Json: Serializing from FArchive is currently broken
+	/**
+	 * Deserializes a internationalization archive from an archive.
 	 *
+	 * @param Archive The archive to serialize from.
+	 * @param InternationalizationArchive Will hold the content of the deserialized internationalization archive.
 	 * @return true if deserialization was successful, false otherwise.
 	 */
-	virtual bool DeserializeArchive( FArchive& Archive, TSharedRef< FInternationalizationArchive > InternationalizationArchive ) = 0;
+	virtual bool DeserializeArchive( FArchive& Archive, TSharedRef<FInternationalizationArchive> InternationalizationArchive ) = 0;
 
 	/**
-	 * Serializes a Internationalizationarchive into an archive.
-	 *
-	 * @param InternationalizationArchive - The Internationalizationarchive data to serialize.
-	 * @param Archive - The archive to serialize into.
+	 * Serializes a internationalization archive into an archive.
 	 *
+	 * @param InternationalizationArchive The Internationalization archive data to serialize.
+	 * @param Archive The archive to serialize into.
 	 * @return true if serialization was successful, false otherwise.
 	 */
-	virtual bool SerializeArchive( TSharedRef< const FInternationalizationArchive > InternationalizationArchive, FArchive& Archive ) = 0;
+	virtual bool SerializeArchive( TSharedRef<const FInternationalizationArchive> InternationalizationArchive, FArchive& Archive ) = 0;
+#endif
+
+public:
 
-};
\ No newline at end of file
+	/** Virtual destructor. */
+	virtual ~IInternationalizationArchiveSerializer() { }
+};
diff --git a/Engine/Source/Runtime/Core/Public/Internationalization/IInternationalizationManifestSerializer.h b/Engine/Source/Runtime/Core/Public/Internationalization/IInternationalizationManifestSerializer.h
--- a/Engine/Source/Runtime/Core/Public/Internationalization/IInternationalizationManifestSerializer.h
+++ b/Engine/Source/Runtime/Core/Public/Internationalization/IInternationalizationManifestSerializer.h
@@ -2,8 +2,10 @@
 
 #pragma once
 
+
 class FInternationalizationManifest;
 
+
 /**
  * Interface for Internationalization manifest serializers.
  */
@@ -12,28 +14,46 @@ class IInternationalizationManifestSerializer
 public:
 
 	/**
-	 * Virtual destructor.
+	 * Deserializes a Internationalization manifest from a string.
+	 *
+	 * @param InStr The string to serialize from.
+	 * @param Manifest The populated Internationalization manifest.
+	 * @return true if deserialization was successful, false otherwise.
 	 */
-	virtual ~IInternationalizationManifestSerializer( ) { }
+	virtual bool DeserializeManifest( const FString& InStr, TSharedRef<FInternationalizationManifest> Manifest ) = 0;
 
 	/**
-	 * Deserializes a manifest from an archive.
+	 * Serializes a Internationalization manifest to a string.
 	 *
-	 * @param Archive - The archive to serialize from.
-	 * @param Manifest - Will hold the content of the deserialized manifest.
+	 * @param Manifest The Internationalization manifest data to serialize.
+	 * @param Str The string to serialize into.
+	 * @return true if serialization was successful, false otherwise.
+	 */
+	virtual bool SerializeManifest( TSharedRef<const FInternationalizationManifest> Manifest, FString& Str ) = 0;
+
+
+#if 0 // @todo Json: Serializing from FArchive is currently broken
+	/**
+	 * Deserializes a manifest from an archive.
 	 *
+	 * @param Archive The archive to serialize from.
+	 * @param Manifest Will hold the content of the deserialized manifest.
 	 * @return true if deserialization was successful, false otherwise.
 	 */
-	virtual bool DeserializeManifest( FArchive& Archive, TSharedRef< FInternationalizationManifest > Manifest ) = 0;
+	virtual bool DeserializeManifest( FArchive& Archive, TSharedRef<FInternationalizationManifest> Manifest ) = 0;
 
 	/**
 	 * Serializes a manifest into an archive.
 	 *
-	 * @param Manifest - The manifest data to serialize.
-	 * @param Archive - The archive to serialize into.
-	 *
+	 * @param Manifest The manifest data to serialize.
+	 * @param Archive The archive to serialize into.
 	 * @return true if serialization was successful, false otherwise.
 	 */
-	virtual bool SerializeManifest( TSharedRef< const FInternationalizationManifest > Manifest, FArchive& Archive ) = 0;
+	virtual bool SerializeManifest( TSharedRef<const FInternationalizationManifest> Manifest, FArchive& Archive ) = 0;
+#endif
+
+public:
 
-};
\ No newline at end of file
+	/** Virtual destructor. */
+	virtual ~IInternationalizationManifestSerializer() { }
+};
diff --git a/Engine/Source/Runtime/Internationalization/Private/Serialization/JsonInternationalizationArchiveSerializer.cpp b/Engine/Source/Runtime/Internationalization/Private/Serialization/JsonInternationalizationArchiveSerializer.cpp
--- a/Engine/Source/Runtime/Internationalization/Private/Serialization/JsonInternationalizationArchiveSerializer.cpp
+++ b/Engine/Source/Runtime/Internationalization/Private/Serialization/JsonInternationalizationArchiveSerializer.cpp
@@ -59,6 +59,7 @@ struct FCompareStructuredArchiveEntryByNamespace
 };
 
 
+#if 0 // @todo Json: Serializing from FArchive is currently broken
 bool FJsonInternationalizationArchiveSerializer::DeserializeArchive( FArchive& Archive, TSharedRef< FInternationalizationArchive > InternationalizationArchive )
 {
 	TSharedPtr< FJsonObject > JsonArchiveObj;
@@ -72,6 +73,7 @@ bool FJsonInternationalizationArchiveSerializer::DeserializeArchive( FArchive& A
 
 	return bExecSuccessful;
 }
+#endif
 
 
 bool FJsonInternationalizationArchiveSerializer::DeserializeArchive( const FString& InStr, TSharedRef< FInternationalizationArchive > InternationalizationArchive )
@@ -95,6 +97,7 @@ bool FJsonInternationalizationArchiveSerializer::DeserializeArchive( TSharedRef<
 }
 
 
+#if 0 // @todo Json: Serializing from FArchive is currently broken
 bool FJsonInternationalizationArchiveSerializer::SerializeArchive( TSharedRef< const FInternationalizationArchive > InternationalizationArchive, FArchive& Archive )
 {
 	TSharedRef< FJsonObject > JsonArchiveObj = MakeShareable( new FJsonObject );
@@ -108,6 +111,7 @@ bool FJsonInternationalizationArchiveSerializer::SerializeArchive( TSharedRef< c
 	}
 	return bExecSuccessful;
 }
+#endif
 
 
 bool FJsonInternationalizationArchiveSerializer::SerializeArchive( TSharedRef< const FInternationalizationArchive > InternationalizationArchive, FString& Str )
diff --git a/Engine/Source/Runtime/Internationalization/Private/Serialization/JsonInternationalizationManifestSerializer.cpp b/Engine/Source/Runtime/Internationalization/Private/Serialization/JsonInternationalizationManifestSerializer.cpp
--- a/Engine/Source/Runtime/Internationalization/Private/Serialization/JsonInternationalizationManifestSerializer.cpp
+++ b/Engine/Source/Runtime/Internationalization/Private/Serialization/JsonInternationalizationManifestSerializer.cpp
@@ -46,6 +46,7 @@ struct FCompareStructuredEntryByNamespace
 };
 
 
+#if 0 // @todo Json: Serializing from FArchive is currently broken
 bool FJsonInternationalizationManifestSerializer::DeserializeManifest( FArchive& Archive, TSharedRef< FInternationalizationManifest > Manifest )
 {
 	TSharedPtr< FJsonObject > JsonManifestObj;
@@ -59,6 +60,7 @@ bool FJsonInternationalizationManifestSerializer::DeserializeManifest( FArchive&
 
 	return bExecSuccessful;
 }
+#endif
 
 
 bool FJsonInternationalizationManifestSerializer::DeserializeManifest( const FString& InStr, TSharedRef< FInternationalizationManifest > Manifest )
@@ -82,6 +84,7 @@ bool FJsonInternationalizationManifestSerializer::DeserializeManifest( TSharedRe
 }
 
 
+#if 0 // @todo Json: Serializing from FArchive is currently broken
 bool FJsonInternationalizationManifestSerializer::SerializeManifest( TSharedRef< const FInternationalizationManifest > Manifest, FArchive& Archive )
 {
 	TSharedRef< FJsonObject > JsonManifestObj = MakeShareable( new FJsonObject );
@@ -95,6 +98,7 @@ bool FJsonInternationalizationManifestSerializer::SerializeManifest( TSharedRef<
 	}
 	return bExecSuccessful;
 }
+#endif
 
 
 bool FJsonInternationalizationManifestSerializer::SerializeManifest( TSharedRef< const FInternationalizationManifest > Manifest, FString& Str )
diff --git a/Engine/Source/Runtime/Internationalization/Public/Serialization/JsonInternationalizationArchiveSerializer.h b/Engine/Source/Runtime/Internationalization/Public/Serialization/JsonInternationalizationArchiveSerializer.h
--- a/Engine/Source/Runtime/Internationalization/Public/Serialization/JsonInternationalizationArchiveSerializer.h
+++ b/Engine/Source/Runtime/Internationalization/Public/Serialization/JsonInternationalizationArchiveSerializer.h
@@ -33,65 +33,37 @@ public:
 	/** Default constructor. */
 	FJsonInternationalizationArchiveSerializer() { }
 
-	/**
-	 * Deserializes a Internationalization archive from archive.
-	 *
-	 * @param Archive The archive to serialize from.
-	 * @param InternationalizationArchive The populated Internationalization archive.
-	 *
-	 * @return true if deserialization was successful, false otherwise.
-	 */
-	virtual bool DeserializeArchive( FArchive& Archive, TSharedRef< FInternationalizationArchive > InternationalizationArchive ) override;
-
-	/**
-	 * Deserializes a Internationalization archive from a string.
-	 *
-	 * @param InStr The string to serialize from.
-	 * @param InternationalizationArchive The populated Internationalization archive.
-	 *
-	 * @return true if deserialization was successful, false otherwise.
-	 */
-	virtual bool DeserializeArchive( const FString& InStr, TSharedRef< FInternationalizationArchive > InternationalizationArchive );
+public:
 
 	/**
 	 * Deserializes a Internationalization archive from a JSON object.
 	 *
 	 * @param InJsonObj The JSON object to serialize from.
 	 * @param InternationalizationArchive The populated Internationalization archive.
-	 *
 	 * @return true if deserialization was successful, false otherwise.
 	 */
-	virtual bool DeserializeArchive( TSharedRef< FJsonObject > InJsonObj, TSharedRef< FInternationalizationArchive > InternationalizationArchive );
-
-	/**
-	 * Serializes a Internationalization archive into an archive.
-	 *
-	 * @param InternationalizationArchive The Internationalization archive data to serialize.
-	 * @param Archive The archive to serialize into.
-	 *
-	 * @return true if serialization was successful, false otherwise.
-	 */
-	virtual bool SerializeArchive( TSharedRef< const FInternationalizationArchive > InternationalizationArchive, FArchive& Archive ) override;
-
-	/**
-	 * Serializes a Internationalization archive to a string.
-	 *
-	 * @param InternationalizationArchive The Internationalization archive data to serialize.
-	 * @param Str The string to serialize into.
-	 *
-	 * @return true if serialization was successful, false otherwise.
-	 */
-	virtual bool SerializeArchive( TSharedRef< const FInternationalizationArchive > InternationalizationArchive, FString& Str );
+	bool DeserializeArchive( TSharedRef< FJsonObject > InJsonObj, TSharedRef< FInternationalizationArchive > InternationalizationArchive );
 
 	/**
 	 * Serializes a Internationalization archive to a JSON object.
 	 *
 	 * @param InternationalizationArchive The Internationalization archive data to serialize.
 	 * @param JsonObj The JSON object to serialize into.
-	 *
 	 * @return true if serialization was successful, false otherwise.
 	 */
-	virtual bool SerializeArchive( TSharedRef< const FInternationalizationArchive > InternationalizationArchive, TSharedRef< FJsonObject > JsonObj );
+	bool SerializeArchive( TSharedRef< const FInternationalizationArchive > InternationalizationArchive, TSharedRef< FJsonObject > JsonObj );
+
+public:
+
+	// IInternationalizationArchiveSerializer interface
+
+	virtual bool DeserializeArchive( const FString& InStr, TSharedRef< FInternationalizationArchive > InternationalizationArchive ) override;
+	virtual bool SerializeArchive( TSharedRef< const FInternationalizationArchive > InternationalizationArchive, FString& Str ) override;
+
+#if 0 // @todo Json: Serializing from FArchive is currently broken
+	virtual bool DeserializeArchive( FArchive& Archive, TSharedRef< FInternationalizationArchive > InternationalizationArchive ) override;
+	virtual bool SerializeArchive( TSharedRef< const FInternationalizationArchive > InternationalizationArchive, FArchive& Archive ) override;
+#endif
 
 protected:
 
@@ -100,7 +72,6 @@ protected:
 	 *
 	 * @param InJsonObj The JSON object to serialize from.
 	 * @param InternationalizationArchive The Internationalization archive that will store the data.
-	 *
 	 * @return true if deserialization was successful, false otherwise.
 	 */
 	bool DeserializeInternal( TSharedRef< FJsonObject > InJsonObj, TSharedRef< FInternationalizationArchive > InternationalizationArchive );
@@ -110,7 +81,6 @@ protected:
 	 *
 	 * @param InInternationalizationArchive The Internationalization archive object to serialize from.
 	 * @param JsonObj The Json object that will store the data.
-	 *
 	 * @return true if serialization was successful, false otherwise.
 	 */
 	bool SerializeInternal( TSharedRef< const FInternationalizationArchive > InInternationalizationArchive, TSharedRef< FJsonObject > JsonObj );
@@ -121,7 +91,6 @@ protected:
 	 * @param InJsonObj The JSON object.
 	 * @param ParentNamespace The namespace of the parent JSON object.
 	 * @param InternationalizationArchive The Internationalization archive that will store the data.
-	 *
 	 * @return true if successful, false otherwise.
 	 */
 	bool JsonObjToArchive( TSharedRef< FJsonObject > InJsonObj, FString ParentNamespace, TSharedRef< FInternationalizationArchive > InternationalizationArchive );
diff --git a/Engine/Source/Runtime/Internationalization/Public/Serialization/JsonInternationalizationManifestSerializer.h b/Engine/Source/Runtime/Internationalization/Public/Serialization/JsonInternationalizationManifestSerializer.h
--- a/Engine/Source/Runtime/Internationalization/Public/Serialization/JsonInternationalizationManifestSerializer.h
+++ b/Engine/Source/Runtime/Internationalization/Public/Serialization/JsonInternationalizationManifestSerializer.h
@@ -30,65 +30,37 @@ public:
 	/** Default constructor. */
 	FJsonInternationalizationManifestSerializer() { }
 
-	/**
-	 * Deserializes a Internationalization manifest from archive.
-	 *
-	 * @param Archive The archive to serialize from.
-	 * @param Manifest The populated Internationalization manifest.
-	 *
-	 * @return true if deserialization was successful, false otherwise.
-	 */
-	virtual bool DeserializeManifest( FArchive& Archive, TSharedRef< FInternationalizationManifest > Manifest ) override;
-
-	/**
-	 * Deserializes a Internationalization manifest from a string.
-	 *
-	 * @param InStr The string to serialize from.
-	 * @param Manifest The populated Internationalization manifest.
-	 *
-	 * @return true if deserialization was successful, false otherwise.
-	 */
-	virtual bool DeserializeManifest( const FString& InStr, TSharedRef< FInternationalizationManifest > Manifest );
+public:
 
 	/**
 	 * Deserializes a Internationalization manifest from a JSON object.
 	 *
 	 * @param InJsonObj The JSON object to serialize from.
 	 * @param Manifest The populated Internationalization manifest.
-	 *
 	 * @return true if deserialization was successful, false otherwise.
 	 */
-	virtual bool DeserializeManifest( TSharedRef< FJsonObject > InJsonObj, TSharedRef< FInternationalizationManifest > Manifest );
-
-	/**
-	 * Serializes a Internationalization manifest into an archive.
-	 *
-	 * @param Manifest The Internationalization manifest data to serialize.
-	 * @param Archive The archive to serialize into.
-	 *
-	 * @return true if serialization was successful, false otherwise.
-	 */
-	virtual bool SerializeManifest(TSharedRef< const FInternationalizationManifest > Manifest, FArchive& Archive ) override;
-
-	/**
-	 * Serializes a Internationalization manifest to a string.
-	 *
-	 * @param Manifest The Internationalization manifest data to serialize.
-	 * @param Str The string to serialize into.
-	 *
-	 * @return true if serialization was successful, false otherwise.
-	 */
-	virtual bool SerializeManifest( TSharedRef< const FInternationalizationManifest > Manifest, FString& Str );
+	bool DeserializeManifest( TSharedRef< FJsonObject > InJsonObj, TSharedRef< FInternationalizationManifest > Manifest );
 
 	/**
 	 * Serializes a Internationalization manifest to a JSON object.
 	 *
 	 * @param Manifest The Internationalization manifest data to serialize.
 	 * @param JsonObj The JSON object to serialize into.
-	 *
 	 * @return true if serialization was successful, false otherwise.
 	 */
-	virtual bool SerializeManifest( TSharedRef< const FInternationalizationManifest > Manifest, TSharedRef< FJsonObject > JsonObj );
+	bool SerializeManifest( TSharedRef< const FInternationalizationManifest > Manifest, TSharedRef< FJsonObject > JsonObj );
+
+public:
+
+	// IInternationalizationManifestSerializer interface
+
+	virtual bool DeserializeManifest( const FString& InStr, TSharedRef< FInternationalizationManifest > Manifest ) override;
+	virtual bool SerializeManifest( TSharedRef< const FInternationalizationManifest > Manifest, FString& Str ) override;
+
+#if 0 // @todo Json: Serializing from FArchive is currently broken
+	virtual bool DeserializeManifest( FArchive& Archive, TSharedRef< FInternationalizationManifest > Manifest ) override;
+	virtual bool SerializeManifest(TSharedRef< const FInternationalizationManifest > Manifest, FArchive& Archive ) override;
+#endif
 
 protected:
 
