Copying //UE4/Dev-Editor to //UE4/Dev-Main (Source: //UE4/Dev-Editor @ 3082391)

diff --git a/Engine/Source/Runtime/Internationalization/Private/LocTextHelper.cpp b/Engine/Source/Runtime/Internationalization/Private/LocTextHelper.cpp
new file mode 100644
--- /dev/null
+++ b/Engine/Source/Runtime/Internationalization/Private/LocTextHelper.cpp
@@ -0,0 +1,1042 @@
+// Copyright 1998-2016 Epic Games, Inc. All Rights Reserved.
+
+#include "InternationalizationPrivatePCH.h"
+#include "LocTextHelper.h"
+
+#include "Json.h"
+#include "JsonInternationalizationManifestSerializer.h"
+#include "JsonInternationalizationArchiveSerializer.h"
+#include "JsonInternationalizationMetadataSerializer.h"
+
+#define LOCTEXT_NAMESPACE "LocTextHelper"
+
+DEFINE_LOG_CATEGORY_STATIC(LogLocTextHelper, Log, All);
+
+void FLocTextConflicts::AddConflict(const FString& InNamespace, const FString& InKey, const TSharedPtr<FLocMetadataObject>& InKeyMetadata, const FLocItem& InSource, const FString& InSourceLocation)
+{
+	TSharedPtr<FConflict> ExistingEntry = FindEntryByKey(InNamespace, InKey, InKeyMetadata);
+	if (!ExistingEntry.IsValid())
+	{
+		TSharedRef<FConflict> NewEntry = MakeShareable(new FConflict(InNamespace, InKey, InKeyMetadata));
+		EntriesByKey.Add(InKey, NewEntry);
+		ExistingEntry = NewEntry;
+	}
+	ExistingEntry->Add(InSource, InSourceLocation.ReplaceCharWithEscapedChar());
+}
+
+TSharedPtr<FLocTextConflicts::FConflict> FLocTextConflicts::FindEntryByKey(const FString& InNamespace, const FString& InKey, const TSharedPtr<FLocMetadataObject> InKeyMetadata) const
+{
+	TArray<TSharedRef<FConflict>> MatchingEntries;
+	EntriesByKey.MultiFind(InKey, MatchingEntries);
+
+	for (const TSharedRef<FConflict>& Entry : MatchingEntries)
+	{
+		if (Entry->Namespace.Equals(InNamespace, ESearchCase::CaseSensitive))
+		{
+			if (InKeyMetadata.IsValid() != Entry->KeyMetadataObj.IsValid())
+			{
+				continue;
+			}
+			else if ((!InKeyMetadata.IsValid() && !Entry->KeyMetadataObj.IsValid()) || (*InKeyMetadata == *Entry->KeyMetadataObj))
+			{
+				return Entry;
+			}
+		}
+	}
+
+	return nullptr;
+}
+
+FString FLocTextConflicts::GetConflictReport() const
+{
+	FString Report;
+
+	for (const auto& ConflictPair : EntriesByKey)
+	{
+		const TSharedRef<FConflict>& Conflict = ConflictPair.Value;
+		const FString& Namespace = Conflict->Namespace;
+		const FString& Key = Conflict->Key;
+
+		bool bAddToReport = false;
+		TArray<FLocItem> SourceList;
+		Conflict->EntriesBySourceLocation.GenerateValueArray(SourceList);
+		if (SourceList.Num() >= 2)
+		{
+			for (int32 i = 0; i < SourceList.Num() - 1 && !bAddToReport; ++i)
+			{
+				for (int32 j = i + 1; j < SourceList.Num() && !bAddToReport; ++j)
+				{
+					if (!(SourceList[i] == SourceList[j]))
+					{
+						bAddToReport = true;
+					}
+				}
+			}
+		}
+
+		if (bAddToReport)
+		{
+			FString KeyMetadataString = FJsonInternationalizationMetaDataSerializer::MetadataToString(Conflict->KeyMetadataObj);
+			Report += FString::Printf(TEXT("%s - %s %s\n"), *Namespace, *Key, *KeyMetadataString);
+
+			for (auto EntryIter = Conflict->EntriesBySourceLocation.CreateConstIterator(); EntryIter; ++EntryIter)
+			{
+				const FString& SourceLocation = EntryIter.Key();
+				FString ProcessedSourceLocation = FPaths::ConvertRelativePathToFull(SourceLocation);
+				ProcessedSourceLocation.ReplaceInline(TEXT("\\"), TEXT("/"));
+				ProcessedSourceLocation.ReplaceInline(*FPaths::RootDir(), TEXT("/"));
+
+				const FString& SourceText = EntryIter.Value().Text.ReplaceCharWithEscapedChar();
+
+				FString SourceMetadataString = FJsonInternationalizationMetaDataSerializer::MetadataToString(EntryIter.Value().MetadataObj);
+				Report += FString::Printf(TEXT("\t%s - \"%s\" %s\n"), *ProcessedSourceLocation, *SourceText, *SourceMetadataString);
+			}
+			Report += TEXT("\n");
+		}
+	}
+
+	return Report;
+}
+
+FLocTextHelper::FLocTextHelper(TSharedPtr<ILocFileNotifies> InLocFileNotifies)
+	: LocFileNotifies(MoveTemp(InLocFileNotifies))
+{
+}
+
+FLocTextHelper::FLocTextHelper(FString InTargetPath, FString InManifestName, FString InArchiveName, FString InNativeCulture, TArray<FString> InForeignCultures, TSharedPtr<ILocFileNotifies> InLocFileNotifies)
+	: TargetPath(MoveTemp(InTargetPath))
+	, ManifestName(MoveTemp(InManifestName))
+	, ArchiveName(MoveTemp(InArchiveName))
+	, NativeCulture(MoveTemp(InNativeCulture))
+	, ForeignCultures(MoveTemp(InForeignCultures))
+	, LocFileNotifies(MoveTemp(InLocFileNotifies))
+{
+	checkf(!TargetPath.IsEmpty(), TEXT("Target path may not be empty!"));
+	checkf(!ManifestName.IsEmpty(), TEXT("Manifest name may not be empty!"));
+	checkf(!ArchiveName.IsEmpty(), TEXT("Archive name may not be empty!"));
+
+	// Make sure the native culture isn't in the list of foreign cultures
+	if (!NativeCulture.IsEmpty())
+	{
+		ForeignCultures.Remove(NativeCulture);
+	}
+}
+
+bool FLocTextHelper::HasManifest() const
+{
+	return Manifest.IsValid();
+}
+
+bool FLocTextHelper::LoadManifest(const ELocTextHelperLoadFlags InLoadFlags, FText* OutError)
+{
+	const FString ManifestFilePath = TargetPath / ManifestName;
+	return LoadManifest(ManifestFilePath, InLoadFlags, OutError);
+}
+
+bool FLocTextHelper::LoadManifest(const FString& InManifestFilePath, const ELocTextHelperLoadFlags InLoadFlags, FText* OutError)
+{
+	Manifest.Reset();
+	Manifest = LoadManifestImpl(InManifestFilePath, InLoadFlags, OutError);
+	return Manifest.IsValid();
+}
+
+bool FLocTextHelper::SaveManifest(FText* OutError) const
+{
+	const FString ManifestFilePath = TargetPath / ManifestName;
+	return SaveManifest(ManifestFilePath, OutError);
+}
+
+bool FLocTextHelper::SaveManifest(const FString& InManifestFilePath, FText* OutError) const
+{
+	if (!Manifest.IsValid())
+	{
+		if (OutError)
+		{
+			*OutError = FText::Format(LOCTEXT("Error_SaveManifest_NoManifest", "Failed to save file '{0}' as there is no manifest instance to save."), FText::FromString(InManifestFilePath));
+		}
+		return false;
+	}
+
+	return SaveManifestImpl(Manifest.ToSharedRef(), InManifestFilePath, OutError);
+}
+
+void FLocTextHelper::TrimManifest()
+{
+	if (Dependencies.Num() > 0)
+	{
+		// We'll generate a new manifest by only including items that are not in the dependencies
+		TSharedRef<FInternationalizationManifest> TrimmedManifest = MakeShareable(new FInternationalizationManifest());
+
+		for (FManifestEntryByStringContainer::TConstIterator It(Manifest->GetEntriesByKeyIterator()); It; ++It)
+		{
+			const TSharedRef<FManifestEntry> ManifestEntry = It.Value();
+
+			for (const FManifestContext& Context : ManifestEntry->Contexts)
+			{
+				FString DependencyFileName;
+				const TSharedPtr<FManifestEntry> DependencyEntry = FindDependencyEntry(ManifestEntry->Namespace, Context, &DependencyFileName);
+
+				if (DependencyEntry.IsValid())
+				{
+					if (!(DependencyEntry->Source.IsExactMatch(ManifestEntry->Source)))
+					{
+						// There is a dependency manifest entry that has the same namespace and keys as our main manifest entry but the source text differs.
+						FString Message = SanitizeLogOutput(
+							FString::Printf(TEXT("Found previously entered localized string [%s] %s %s=\"%s\" %s. It was previously \"%s\" %s in dependency manifest %s."),
+								*ManifestEntry->Namespace,
+								*Context.Key,
+								*FJsonInternationalizationMetaDataSerializer::MetadataToString(Context.KeyMetadataObj),
+								*ManifestEntry->Source.Text,
+								*FJsonInternationalizationMetaDataSerializer::MetadataToString(ManifestEntry->Source.MetadataObj),
+								*DependencyEntry->Source.Text,
+								*FJsonInternationalizationMetaDataSerializer::MetadataToString(DependencyEntry->Source.MetadataObj),
+								*DependencyFileName
+								)
+							);
+						UE_LOG(LogLocTextHelper, Warning, TEXT("%s"), *Message);
+
+						ConflictTracker.AddConflict(ManifestEntry->Namespace, Context.Key, Context.KeyMetadataObj, ManifestEntry->Source, *Context.SourceLocation);
+
+						const FManifestContext* ConflictingContext = DependencyEntry->FindContext(Context.Key, Context.KeyMetadataObj);
+						const FString DependencyEntryFullSrcLoc = (!DependencyFileName.IsEmpty()) ? DependencyFileName : ConflictingContext->SourceLocation;
+
+						ConflictTracker.AddConflict(ManifestEntry->Namespace, Context.Key, Context.KeyMetadataObj, DependencyEntry->Source, DependencyEntryFullSrcLoc);
+					}
+				}
+				else
+				{
+					// Since we did not find any entries in the dependencies list that match, we'll add to the new manifest
+					const bool bAddSuccessful = TrimmedManifest->AddSource(ManifestEntry->Namespace, ManifestEntry->Source, Context);
+					if (!bAddSuccessful)
+					{
+						UE_LOG(LogLocTextHelper, Error, TEXT("Could not process localized string: %s [%s] %s=\"%s\" %s."),
+							*ManifestEntry->Namespace,
+							*Context.Key,
+							*ManifestEntry->Source.Text,
+							*FJsonInternationalizationMetaDataSerializer::MetadataToString(ManifestEntry->Source.MetadataObj)
+							);
+					}
+				}
+			}
+		}
+
+		Manifest = TrimmedManifest;
+	}
+}
+
+bool FLocTextHelper::HasNativeArchive() const
+{
+	TSharedPtr<FInternationalizationArchive> Archive = (NativeCulture.IsEmpty()) ? TSharedPtr<FInternationalizationArchive>() : Archives.FindRef(NativeCulture);
+	return Archive.IsValid();
+}
+
+bool FLocTextHelper::LoadNativeArchive(const ELocTextHelperLoadFlags InLoadFlags, FText* OutError)
+{
+	const FString ArchiveFilePath = TargetPath / NativeCulture / ArchiveName;
+	return LoadNativeArchive(ArchiveFilePath, InLoadFlags, OutError);
+}
+
+bool FLocTextHelper::LoadNativeArchive(const FString& InArchiveFilePath, const ELocTextHelperLoadFlags InLoadFlags, FText* OutError)
+{
+	checkf(!NativeCulture.IsEmpty(), TEXT("Attempted to load the native culture archive file, but no native culture was set during construction!"));
+	checkf(Manifest.IsValid(), TEXT("Attempted to load the native culture archive file, but no manifest has been loaded!"));
+
+	Archives.Remove(NativeCulture);
+	
+	TSharedPtr<FInternationalizationArchive> Archive = LoadArchiveImpl(InArchiveFilePath, InLoadFlags, OutError);
+	if (Archive.IsValid())
+	{
+		Archives.Add(NativeCulture, Archive);
+		return true;
+	}
+	return false;
+}
+
+bool FLocTextHelper::SaveNativeArchive(FText* OutError) const
+{
+	const FString ArchiveFilePath = TargetPath / NativeCulture / ArchiveName;
+	return SaveNativeArchive(ArchiveFilePath, OutError);
+}
+
+bool FLocTextHelper::SaveNativeArchive(const FString& InArchiveFilePath, FText* OutError) const
+{
+	TSharedPtr<FInternationalizationArchive> Archive = (NativeCulture.IsEmpty()) ? TSharedPtr<FInternationalizationArchive>() : Archives.FindRef(NativeCulture);
+	if (!Archive.IsValid())
+	{
+		if (OutError)
+		{
+			*OutError = FText::Format(LOCTEXT("Error_SaveArchive_NoArchive", "Failed to save file '{0}' as there is no archive instance to save."), FText::FromString(InArchiveFilePath));
+		}
+		return false;
+	}
+
+	return SaveArchiveImpl(Archive.ToSharedRef(), InArchiveFilePath, OutError);
+}
+
+bool FLocTextHelper::HasForeignArchive(const FString& InCulture) const
+{
+	TSharedPtr<FInternationalizationArchive> Archive = Archives.FindRef(InCulture);
+	return Archive.IsValid();
+}
+
+bool FLocTextHelper::LoadForeignArchive(const FString& InCulture, const ELocTextHelperLoadFlags InLoadFlags, FText* OutError)
+{
+	const FString ArchiveFilePath = TargetPath / InCulture / ArchiveName;
+	return LoadForeignArchive(InCulture, ArchiveFilePath, InLoadFlags, OutError);
+}
+
+bool FLocTextHelper::LoadForeignArchive(const FString& InCulture, const FString& InArchiveFilePath, const ELocTextHelperLoadFlags InLoadFlags, FText* OutError)
+{
+	checkf(ForeignCultures.Contains(InCulture), TEXT("Attempted to load a foreign culture archive file, but the given culture (%s) wasn't set during construction!"), *InCulture);
+	checkf(Manifest.IsValid(), TEXT("Attempted to load a foreign culture archive file, but no manifest has been loaded!"));
+
+	Archives.Remove(InCulture);
+
+	TSharedPtr<FInternationalizationArchive> Archive = LoadArchiveImpl(InArchiveFilePath, InLoadFlags, OutError);
+	if (Archive.IsValid())
+	{
+		Archives.Add(InCulture, Archive);
+		return true;
+	}
+	return false;
+}
+
+bool FLocTextHelper::SaveForeignArchive(const FString& InCulture, FText* OutError) const
+{
+	const FString ArchiveFilePath = TargetPath / InCulture / ArchiveName;
+	return SaveForeignArchive(InCulture, ArchiveFilePath, OutError);
+}
+
+bool FLocTextHelper::SaveForeignArchive(const FString& InCulture, const FString& InArchiveFilePath, FText* OutError) const
+{
+	TSharedPtr<FInternationalizationArchive> Archive = Archives.FindRef(InCulture);
+	if (!Archive.IsValid())
+	{
+		if (OutError)
+		{
+			*OutError = FText::Format(LOCTEXT("Error_SaveArchive_NoArchive", "Failed to save file '{0}' as there is no archive instance to save."), FText::FromString(InArchiveFilePath));
+		}
+		return false;
+	}
+
+	return SaveArchiveImpl(Archive.ToSharedRef(), InArchiveFilePath, OutError);
+}
+
+bool FLocTextHelper::LoadAllArchives(const ELocTextHelperLoadFlags InLoadFlags, FText* OutError)
+{
+	if (!NativeCulture.IsEmpty() && !LoadNativeArchive(InLoadFlags, OutError))
+	{
+		return false;
+	}
+
+	for (const FString& Culture : ForeignCultures)
+	{
+		if (!LoadForeignArchive(Culture, InLoadFlags, OutError))
+		{
+			return false;
+		}
+	}
+
+	return true;
+}
+
+bool FLocTextHelper::SaveAllArchives(FText* OutError) const
+{
+	if (!NativeCulture.IsEmpty() && !SaveNativeArchive(OutError))
+	{
+		return false;
+	}
+
+	for (const FString& Culture : ForeignCultures)
+	{
+		if (!SaveForeignArchive(Culture, OutError))
+		{
+			return false;
+		}
+	}
+
+	return true;
+}
+
+void FLocTextHelper::TrimArchive(const FString& InCulture)
+{
+	checkf(Manifest.IsValid(), TEXT("Attempted to trim an archive file, but no manifest has been loaded!"));
+
+	TSharedPtr<FInternationalizationArchive> Archive = Archives.FindRef(InCulture);
+	checkf(Archive.IsValid(), TEXT("Attempted to trim an archive file, but no valid archive could be found for '%s'!"), *InCulture);
+
+	TSharedPtr<FInternationalizationArchive> NativeArchive;
+	if (!NativeCulture.IsEmpty() && InCulture != NativeCulture)
+	{
+		NativeArchive = Archives.FindRef(NativeCulture);
+		checkf(NativeArchive.IsValid(), TEXT("Attempted to trim an archive file, but no valid archive could be found for '%s'!"), *NativeCulture);
+	}
+
+	// Copy any translations that match current manifest entries over into the trimmed archive
+	TSharedRef<FInternationalizationArchive> TrimmedArchive = MakeShareable(new FInternationalizationArchive());
+	EnumerateSourceTexts([&](TSharedRef<FManifestEntry> InManifestEntry) -> bool
+	{
+		for (const FManifestContext& Context : InManifestEntry->Contexts)
+		{
+			// Keep any translation for the source text
+			TSharedPtr<FArchiveEntry> ArchiveEntry = Archive->FindEntryByKey(InManifestEntry->Namespace, Context.Key, Context.KeyMetadataObj);
+			if (ArchiveEntry.IsValid())
+			{
+				TrimmedArchive->AddEntry(ArchiveEntry.ToSharedRef());
+			}
+		}
+
+		return true; // continue enumeration
+	}, true);
+
+	Archives.Remove(InCulture);
+	Archives.Add(InCulture, TrimmedArchive);
+}
+
+bool FLocTextHelper::LoadAll(const ELocTextHelperLoadFlags InLoadFlags, FText* OutError)
+{
+	if (!LoadManifest(InLoadFlags, OutError))
+	{
+		return false;
+	}
+
+	return LoadAllArchives(InLoadFlags, OutError);
+}
+
+bool FLocTextHelper::SaveAll(FText* OutError) const
+{
+	if (!SaveManifest(OutError))
+	{
+		return false;
+	}
+
+	return SaveAllArchives(OutError);
+}
+
+bool FLocTextHelper::AddDependency(const FString& InDependencyFilePath, FText* OutError)
+{
+	if (DependencyPaths.Contains(InDependencyFilePath))
+	{
+		return true;
+	}
+
+	TSharedPtr<FInternationalizationManifest> DepManifest = LoadManifestImpl(InDependencyFilePath, ELocTextHelperLoadFlags::Load, OutError);
+	if (DepManifest.IsValid())
+	{
+		DependencyPaths.Add(InDependencyFilePath);
+		Dependencies.Add(DepManifest);
+		return true;
+	}
+
+	return false;
+}
+
+TSharedPtr<FManifestEntry> FLocTextHelper::FindDependencyEntry(const FString& InNamespace, const FString& InKey, const FString* InSourceText, FString* OutDependencyFilePath) const
+{
+	for (int32 DepIndex = 0; DepIndex < Dependencies.Num(); ++DepIndex)
+	{
+		TSharedPtr<FInternationalizationManifest> DepManifest = Dependencies[DepIndex];
+
+		const TSharedPtr<FManifestEntry> DepEntry = DepManifest->FindEntryByKey(InNamespace, InKey, InSourceText);
+		if (DepEntry.IsValid())
+		{
+			if (OutDependencyFilePath)
+			{
+				*OutDependencyFilePath = DependencyPaths[DepIndex];
+			}
+			return DepEntry;
+		}
+	}
+
+	return nullptr;
+}
+
+TSharedPtr<FManifestEntry> FLocTextHelper::FindDependencyEntry(const FString& InNamespace, const FManifestContext& InContext, FString* OutDependencyFilePath) const
+{
+	for (int32 DepIndex = 0; DepIndex < Dependencies.Num(); ++DepIndex)
+	{
+		TSharedPtr<FInternationalizationManifest> DepManifest = Dependencies[DepIndex];
+
+		const TSharedPtr<FManifestEntry> DepEntry = DepManifest->FindEntryByContext(InNamespace, InContext);
+		if (DepEntry.IsValid())
+		{
+			if (OutDependencyFilePath)
+			{
+				*OutDependencyFilePath = DependencyPaths[DepIndex];
+			}
+			return DepEntry;
+		}
+	}
+
+	return nullptr;
+}
+
+bool FLocTextHelper::AddSourceText(const FString& InNamespace, const FLocItem& InSource, const FManifestContext& InContext, const FString* InDescription)
+{
+	checkf(Manifest.IsValid(), TEXT("Attempted to add source text, but no manifest has been loaded!"));
+	
+	bool bAddSuccessful = false;
+
+	// Check if the entry already exists in the manifest or one of the manifest dependencies
+	FString ExistingEntryFileName;
+	TSharedPtr<FManifestEntry> ExistingEntry = Manifest->FindEntryByContext(InNamespace, InContext);
+	if (!ExistingEntry.IsValid())
+	{
+		ExistingEntry = FindDependencyEntry(InNamespace, InContext, &ExistingEntryFileName);
+	}
+
+	if (ExistingEntry.IsValid())
+	{
+		if (InSource.IsExactMatch(ExistingEntry->Source))
+		{
+			bAddSuccessful = true;
+		}
+		else
+		{
+			// Grab the source location of the conflicting context
+			const FManifestContext* ConflictingContext = ExistingEntry->FindContext(InContext.Key, InContext.KeyMetadataObj);
+			const FString ExistingEntrySourceLocation = (!ExistingEntryFileName.IsEmpty()) ? ExistingEntryFileName : ConflictingContext->SourceLocation;
+
+			FString Message = SanitizeLogOutput(
+				FString::Printf(TEXT("Found previously entered localized string: %s [%s] %s %s=\"%s\" %s. It was previously \"%s\" %s in %s."),
+					(InDescription ? **InDescription : *FString()),
+					*InNamespace,
+					*InContext.Key,
+					*FJsonInternationalizationMetaDataSerializer::MetadataToString(InContext.KeyMetadataObj),
+					*InSource.Text,
+					*FJsonInternationalizationMetaDataSerializer::MetadataToString(InSource.MetadataObj),
+					*ExistingEntry->Source.Text,
+					*FJsonInternationalizationMetaDataSerializer::MetadataToString(ExistingEntry->Source.MetadataObj),
+					*ExistingEntrySourceLocation
+					)
+				);
+			UE_LOG(LogLocTextHelper, Warning, TEXT("%s"), *Message);
+
+			ConflictTracker.AddConflict(InNamespace, InContext.Key, InContext.KeyMetadataObj, InSource, InContext.SourceLocation);
+			ConflictTracker.AddConflict(InNamespace, InContext.Key, InContext.KeyMetadataObj, ExistingEntry->Source, ExistingEntrySourceLocation);
+		}
+	}
+	else
+	{
+		bAddSuccessful = Manifest->AddSource(InNamespace, InSource, InContext);
+		if (!bAddSuccessful)
+		{
+			UE_LOG(LogLocTextHelper, Error, TEXT("Could not process localized string: %s [%s] %s=\"%s\" %s."),
+				(InDescription ? **InDescription : *FString()),
+				*InNamespace,
+				*InContext.Key,
+				*InSource.Text,
+				*FJsonInternationalizationMetaDataSerializer::MetadataToString(InSource.MetadataObj)
+				);
+		}
+	}
+
+	return bAddSuccessful;
+}
+
+void FLocTextHelper::UpdateSourceText(const TSharedRef<FManifestEntry>& InOldEntry, TSharedRef<FManifestEntry>& InNewEntry)
+{
+	checkf(Manifest.IsValid(), TEXT("Attempted to update source text, but no manifest has been loaded!"));
+	Manifest->UpdateEntry(InOldEntry, InNewEntry);
+}
+
+TSharedPtr<FManifestEntry> FLocTextHelper::FindSourceText(const FString& InNamespace, const FString& InKey, const FString* InSourceText) const
+{
+	checkf(Manifest.IsValid(), TEXT("Attempted to find source text, but no manifest has been loaded!"));
+	return Manifest->FindEntryByKey(InNamespace, InKey, InSourceText);
+}
+
+TSharedPtr<FManifestEntry> FLocTextHelper::FindSourceText(const FString& InNamespace, const FManifestContext& InContext) const
+{
+	checkf(Manifest.IsValid(), TEXT("Attempted to find source text, but no manifest has been loaded!"));
+	return Manifest->FindEntryByContext(InNamespace, InContext);
+}
+
+void FLocTextHelper::EnumerateSourceTexts(const FEnumerateSourceTextsFuncPtr& InCallback, const bool InCheckDependencies) const
+{
+	checkf(Manifest.IsValid(), TEXT("Attempted to enumerate source texts, but no manifest has been loaded!"));
+
+	for (FManifestEntryByStringContainer::TConstIterator It(Manifest->GetEntriesByKeyIterator()); It; ++It)
+	{
+		const TSharedRef<FManifestEntry> ManifestEntry = It.Value();
+
+		bool bShouldEnumerate = true;
+		if (InCheckDependencies)
+		{
+			for (const TSharedPtr<FInternationalizationManifest>& DepManifest : Dependencies)
+			{
+				const TSharedPtr<FManifestEntry> DepEntry = DepManifest->FindEntryBySource(ManifestEntry->Namespace, ManifestEntry->Source);
+				if (DepEntry.IsValid())
+				{
+					bShouldEnumerate = false;
+					break;
+				}
+			}
+		}
+
+		if (bShouldEnumerate && !InCallback(ManifestEntry))
+		{
+			break;
+		}
+	}
+}
+
+bool FLocTextHelper::AddTranslation(const FString& InCulture, const FString& InNamespace, const FString& InKey, const TSharedPtr<FLocMetadataObject>& InKeyMetadataObj, const FLocItem& InSource, const FLocItem& InTranslation, const bool InOptional)
+{
+	TSharedPtr<FInternationalizationArchive> Archive = Archives.FindRef(InCulture);
+	checkf(Archive.IsValid(), TEXT("Attempted to add a translation, but no valid archive could be found for '%s'!"), *InCulture);
+	return Archive->AddEntry(InNamespace, InKey, InSource, InTranslation, InKeyMetadataObj, InOptional);
+}
+
+bool FLocTextHelper::AddTranslation(const FString& InCulture, const TSharedRef<FArchiveEntry>& InEntry)
+{
+	TSharedPtr<FInternationalizationArchive> Archive = Archives.FindRef(InCulture);
+	checkf(Archive.IsValid(), TEXT("Attempted to add a translation, but no valid archive could be found for '%s'!"), *InCulture);
+	return Archive->AddEntry(InEntry);
+}
+
+bool FLocTextHelper::UpdateTranslation(const FString& InCulture, const FString& InNamespace, const FString& InKey, const TSharedPtr<FLocMetadataObject>& InKeyMetadataObj, const FLocItem& InSource, const FLocItem& InTranslation)
+{
+	TSharedPtr<FInternationalizationArchive> Archive = Archives.FindRef(InCulture);
+	checkf(Archive.IsValid(), TEXT("Attempted to update a translation, but no valid archive could be found for '%s'!"), *InCulture);
+	return Archive->SetTranslation(InNamespace, InKey, InSource, InTranslation, InKeyMetadataObj);
+}
+
+void FLocTextHelper::UpdateTranslation(const FString& InCulture, const TSharedRef<FArchiveEntry>& InOldEntry, const TSharedRef<FArchiveEntry>& InNewEntry)
+{
+	TSharedPtr<FInternationalizationArchive> Archive = Archives.FindRef(InCulture);
+	checkf(Archive.IsValid(), TEXT("Attempted to update a translation, but no valid archive could be found for '%s'!"), *InCulture);
+	Archive->UpdateEntry(InOldEntry, InNewEntry);
+}
+
+bool FLocTextHelper::ImportTranslation(const FString& InCulture, const FString& InNamespace, const FString& InKey, const TSharedPtr<FLocMetadataObject> InKeyMetadataObj, const FLocItem& InSource, const FLocItem& InTranslation, const bool InOptional)
+{
+	TSharedPtr<FInternationalizationArchive> Archive = Archives.FindRef(InCulture);
+	checkf(Archive.IsValid(), TEXT("Attempted to update a translation, but no valid archive could be found for '%s'!"), *InCulture);
+
+	// First try and update an existing entry...
+	if (Archive->SetTranslation(InNamespace, InKey, InSource, InTranslation, InKeyMetadataObj))
+	{
+		return true;
+	}
+
+	// ... failing that, try to add a new entry
+	return Archive->AddEntry(InNamespace, InKey, InSource, InTranslation, InKeyMetadataObj, InOptional);
+}
+
+TSharedPtr<FArchiveEntry> FLocTextHelper::FindTranslation(const FString& InCulture, const FString& InNamespace, const FString& InKey, const TSharedPtr<FLocMetadataObject> InKeyMetadataObj) const
+{
+	return FindTranslationImpl(InCulture, InNamespace, InKey, InKeyMetadataObj);
+}
+
+void FLocTextHelper::EnumerateTranslations(const FString& InCulture, const FEnumerateTranslationsFuncPtr& InCallback, const bool InCheckDependencies) const
+{
+	TSharedPtr<FInternationalizationArchive> Archive = Archives.FindRef(InCulture);
+	checkf(Archive.IsValid(), TEXT("Attempted to enumerate translations, but no valid archive could be found for '%s'!"), *InCulture);
+
+	EnumerateSourceTexts([&](TSharedRef<FManifestEntry> InManifestEntry) -> bool
+	{
+		bool bContinue = true;
+		
+		for (const FManifestContext& ManifestContext : InManifestEntry->Contexts)
+		{
+			TSharedPtr<FArchiveEntry> ArchiveEntry = FindTranslation(InCulture, InManifestEntry->Namespace, ManifestContext.Key, ManifestContext.KeyMetadataObj);
+			if (ArchiveEntry.IsValid() && !InCallback(ArchiveEntry.ToSharedRef()))
+			{
+				bContinue = false;
+				break;
+			}
+		}
+
+		return bContinue;
+	}, InCheckDependencies);
+}
+
+void FLocTextHelper::GetExportText(const FString& InCulture, const FString& InNamespace, const FString& InKey, const TSharedPtr<FLocMetadataObject> InKeyMetadataObj, const ELocTextExportSourceMethod InSourceMethod, const FLocItem& InSource, FLocItem& OutSource, FLocItem& OutTranslation) const
+{
+	// Default to the raw source text for the case where we're not using native translations as source
+	OutSource = InSource;
+	OutTranslation = FLocItem();
+
+	if (InSourceMethod == ELocTextExportSourceMethod::NativeText && !NativeCulture.IsEmpty() && InCulture != NativeCulture)
+	{
+		TSharedPtr<FArchiveEntry> NativeArchiveEntry = FindTranslationImpl(NativeCulture, InNamespace, InKey, InKeyMetadataObj);
+		if (NativeArchiveEntry.IsValid() && !NativeArchiveEntry->Source.IsExactMatch(NativeArchiveEntry->Translation))
+		{
+			// Use the native translation as the source
+			OutSource = NativeArchiveEntry->Translation;
+		}
+	}
+
+	TSharedPtr<FArchiveEntry> ArchiveEntry = FindTranslationImpl(InCulture, InNamespace, InKey, InKeyMetadataObj);
+	if (ArchiveEntry.IsValid())
+	{
+		// Set the export text to use the current translation if the entry source matches the export source
+		if (ArchiveEntry->Source.IsExactMatch(OutSource))
+		{
+			OutTranslation = ArchiveEntry->Translation;
+		}
+	}
+	
+	// We use the source text as the default translation for the native culture
+	if (OutTranslation.Text.IsEmpty() && !NativeCulture.IsEmpty() && InCulture == NativeCulture)
+	{
+		OutTranslation = OutSource;
+	}
+}
+
+void FLocTextHelper::GetRuntimeText(const FString& InCulture, const FString& InNamespace, const FString& InKey, const TSharedPtr<FLocMetadataObject> InKeyMetadataObj, const ELocTextExportSourceMethod InSourceMethod, const FLocItem& InSource, FLocItem& OutTranslation, const bool bSkipSourceCheck) const
+{
+	OutTranslation = InSource;
+
+	TSharedPtr<FArchiveEntry> ArchiveEntry = FindTranslationImpl(InCulture, InNamespace, InKey, InKeyMetadataObj);
+	if (ArchiveEntry.IsValid())
+	{
+		if (bSkipSourceCheck)
+		{
+			// Set the export text to use the current translation
+			OutTranslation = ArchiveEntry->Translation;
+		}
+		else
+		{
+			FLocItem ExpectedSource = InSource;
+
+			if (InSourceMethod == ELocTextExportSourceMethod::NativeText && !NativeCulture.IsEmpty() && InCulture != NativeCulture)
+			{
+				TSharedPtr<FArchiveEntry> NativeArchiveEntry = FindTranslationImpl(NativeCulture, InNamespace, InKey, InKeyMetadataObj);
+				if (NativeArchiveEntry.IsValid() && !NativeArchiveEntry->Source.IsExactMatch(NativeArchiveEntry->Translation))
+				{
+					// Use the native translation as the source
+					ExpectedSource = NativeArchiveEntry->Translation;
+				}
+			}
+
+			if (ArchiveEntry->Source.IsExactMatch(ExpectedSource))
+			{
+				// Set the export text to use the current translation
+				OutTranslation = ArchiveEntry->Translation;
+			}
+		}
+	}
+}
+
+FString FLocTextHelper::GetConflictReport() const
+{
+	return ConflictTracker.GetConflictReport();
+}
+
+FString FLocTextHelper::SanitizeLogOutput(const FString& InString)
+{
+	if (InString.IsEmpty())
+	{
+		return InString;
+	}
+
+	FString ResultStr = InString.ReplaceCharWithEscapedChar();
+	if (!GIsBuildMachine)
+	{
+		return ResultStr;
+	}
+
+	static TArray<FString> ErrorStrs;
+	if (ErrorStrs.Num() == 0)
+	{
+		ErrorStrs.Add(FString(TEXT("Error")));
+		ErrorStrs.Add(FString(TEXT("Failed")));
+		ErrorStrs.Add(FString(TEXT("[BEROR]")));
+		ErrorStrs.Add(FString(TEXT("Utility finished with exit code: -1")));
+		ErrorStrs.Add(FString(TEXT("is not recognized as an internal or external command")));
+		ErrorStrs.Add(FString(TEXT("Could not open solution: ")));
+		ErrorStrs.Add(FString(TEXT("Parameter format not correct")));
+		ErrorStrs.Add(FString(TEXT("Another build is already started on this computer.")));
+		ErrorStrs.Add(FString(TEXT("Sorry but the link was not completed because memory was exhausted.")));
+		ErrorStrs.Add(FString(TEXT("simply rerunning the compiler might fix this problem")));
+		ErrorStrs.Add(FString(TEXT("No connection could be made because the target machine actively refused")));
+		ErrorStrs.Add(FString(TEXT("Internal Linker Exception:")));
+		ErrorStrs.Add(FString(TEXT(": warning LNK4019: corrupt string table")));
+		ErrorStrs.Add(FString(TEXT("Proxy could not update its cache")));
+		ErrorStrs.Add(FString(TEXT("You have not agreed to the Xcode license agreements")));
+		ErrorStrs.Add(FString(TEXT("Connection to build service terminated")));
+		ErrorStrs.Add(FString(TEXT("cannot execute binary file")));
+		ErrorStrs.Add(FString(TEXT("Invalid solution configuration")));
+		ErrorStrs.Add(FString(TEXT("is from a previous version of this application and must be converted in order to build")));
+		ErrorStrs.Add(FString(TEXT("This computer has not been authenticated for your account using Steam Guard")));
+		ErrorStrs.Add(FString(TEXT("invalid name for SPA section")));
+		ErrorStrs.Add(FString(TEXT(": Invalid file name, ")));
+		ErrorStrs.Add(FString(TEXT("The specified PFX file do not exist. Aborting")));
+		ErrorStrs.Add(FString(TEXT("binary is not found. Aborting")));
+		ErrorStrs.Add(FString(TEXT("Input file not found: ")));
+		ErrorStrs.Add(FString(TEXT("An exception occurred during merging:")));
+		ErrorStrs.Add(FString(TEXT("Install the 'Microsoft Windows SDK for Windows 7 and .NET Framework 3.5 SP1'")));
+		ErrorStrs.Add(FString(TEXT("is less than package's new version 0x")));
+		ErrorStrs.Add(FString(TEXT("current engine version is older than version the package was originally saved with")));
+		ErrorStrs.Add(FString(TEXT("exceeds maximum length")));
+		ErrorStrs.Add(FString(TEXT("can't edit exclusive file already opened")));
+	}
+
+	for (int32 ErrStrIdx = 0; ErrStrIdx < ErrorStrs.Num(); ErrStrIdx++)
+	{
+		FString& FindStr = ErrorStrs[ErrStrIdx];
+		FString ReplaceStr = FString::Printf(TEXT("%s %s"), *FindStr.Left(1), *FindStr.RightChop(1));
+
+		ResultStr.ReplaceInline(*FindStr, *ReplaceStr);
+	}
+
+	return ResultStr;
+}
+
+bool FLocTextHelper::FindKeysForLegacyTranslation(const FString& InCulture, const FString& InNamespace, const FString& InSource, const TSharedPtr<FLocMetadataObject> InKeyMetadataObj, TArray<FString>& OutKeys) const
+{
+	checkf(Manifest.IsValid(), TEXT("Attempted to find a key for a legacy translation, but no manifest has been loaded!"));
+
+	TSharedPtr<FInternationalizationArchive> NativeArchive;
+	if (!NativeCulture.IsEmpty() && InCulture != NativeCulture)
+	{
+		NativeArchive = Archives.FindRef(NativeCulture);
+		checkf(NativeArchive.IsValid(), TEXT("Attempted to find a key for a legacy translation, but no valid archive could be found for '%s'!"), *NativeCulture);
+	}
+
+	return FindKeysForLegacyTranslation(Manifest.ToSharedRef(), NativeArchive, InNamespace, InSource, InKeyMetadataObj, OutKeys);
+}
+
+bool FLocTextHelper::FindKeysForLegacyTranslation(const TSharedRef<const FInternationalizationManifest>& InManifest, const TSharedPtr<const FInternationalizationArchive>& InNativeArchive, const FString& InNamespace, const FString& InSource, const TSharedPtr<FLocMetadataObject> InKeyMetadataObj, TArray<FString>& OutKeys)
+{
+	FString RealSourceText = InSource;
+
+	// The source text may be a native translation, so we first need to check the native archive to find the real source text that will exist in the manifest
+	if (InNativeArchive.IsValid())
+	{
+		// We don't maintain a translation -> source mapping, so we just have to brute force it
+		for (FArchiveEntryByStringContainer::TConstIterator It = InNativeArchive->GetEntriesBySourceTextIterator(); It; ++It)
+		{
+			const TSharedRef<FArchiveEntry>& ArchiveEntry = It.Value();
+			if (ArchiveEntry->Namespace.Equals(InNamespace, ESearchCase::CaseSensitive) && ArchiveEntry->Translation.Text.Equals(InSource, ESearchCase::CaseSensitive))
+			{
+				if (!ArchiveEntry->KeyMetadataObj.IsValid() && !InKeyMetadataObj.IsValid())
+				{
+					RealSourceText = ArchiveEntry->Source.Text;
+					break;
+				}
+				else if ((InKeyMetadataObj.IsValid() != ArchiveEntry->KeyMetadataObj.IsValid()))
+				{
+					// If we are in here, we know that one of the metadata entries is null, if the other contains zero entries we will still consider them equivalent.
+					if ((InKeyMetadataObj.IsValid() && InKeyMetadataObj->Values.Num() == 0) || (ArchiveEntry->KeyMetadataObj.IsValid() && ArchiveEntry->KeyMetadataObj->Values.Num() == 0))
+					{
+						RealSourceText = ArchiveEntry->Source.Text;
+						break;
+					}
+				}
+				else if (*ArchiveEntry->KeyMetadataObj == *InKeyMetadataObj)
+				{
+					RealSourceText = ArchiveEntry->Source.Text;
+					break;
+				}
+			}
+		}
+	}
+
+	bool bFoundKeys = false;
+
+	TSharedPtr<FManifestEntry> ManifestEntry = InManifest->FindEntryBySource(InNamespace, FLocItem(RealSourceText));
+	if (ManifestEntry.IsValid())
+	{
+		for (const FManifestContext& Context : ManifestEntry->Contexts)
+		{
+			if (Context.KeyMetadataObj.IsValid() != InKeyMetadataObj.IsValid())
+			{
+				continue;
+			}
+			else if ((!Context.KeyMetadataObj.IsValid() && !InKeyMetadataObj.IsValid()) || (*Context.KeyMetadataObj == *InKeyMetadataObj))
+			{
+				OutKeys.Add(Context.Key);
+				bFoundKeys = true;
+			}
+		}
+	}
+
+	return bFoundKeys;
+}
+
+TSharedPtr<FInternationalizationManifest> FLocTextHelper::LoadManifestImpl(const FString& InManifestFilePath, const ELocTextHelperLoadFlags InLoadFlags, FText* OutError)
+{
+	TSharedPtr<FInternationalizationManifest> LocalManifest = MakeShareable(new FInternationalizationManifest());
+
+	// Attempt to load an existing manifest first
+	if (!!(InLoadFlags & ELocTextHelperLoadFlags::Load))
+	{
+		bool bLoaded = false;
+		const bool bExisted = FPaths::FileExists(InManifestFilePath);
+
+		if (bExisted)
+		{
+			if (LocFileNotifies.IsValid())
+			{
+				LocFileNotifies->PreFileRead(InManifestFilePath);
+			}
+
+			if (FJsonInternationalizationManifestSerializer::DeserializeManifestFromFile(InManifestFilePath, LocalManifest.ToSharedRef()))
+			{
+				bLoaded = true;
+			}
+			else
+			{
+				if (OutError)
+				{
+					*OutError = FText::Format(LOCTEXT("Error_LoadManifest_DeserializeFile", "Failed to deserialize manifest '{0}'."), FText::FromString(InManifestFilePath));
+				}
+			}
+
+			if (LocFileNotifies.IsValid())
+			{
+				LocFileNotifies->PostFileRead(InManifestFilePath);
+			}
+		}
+
+		if (bLoaded)
+		{
+			return LocalManifest;
+		}
+		else if (bExisted)
+		{
+			// Don't allow fallback to Create if the file exists but could not be loaded
+			return nullptr;
+		}
+	}
+
+	// If we're allowed to create a manifest than we can never fail
+	if (!!(InLoadFlags & ELocTextHelperLoadFlags::Create))
+	{
+		return LocalManifest;
+	}
+
+	return nullptr;
+}
+
+bool FLocTextHelper::SaveManifestImpl(const TSharedRef<const FInternationalizationManifest>& InManifest, const FString& InManifestFilePath, FText* OutError) const
+{
+	bool bSaved = false;
+
+	if (LocFileNotifies.IsValid())
+	{
+		LocFileNotifies->PreFileWrite(InManifestFilePath);
+	}
+
+	TSharedRef<FJsonObject> ManifestJsonObject = MakeShareable(new FJsonObject());
+	if (FJsonInternationalizationManifestSerializer::SerializeManifestToFile(InManifest, InManifestFilePath))
+	{
+		bSaved = true;
+	}
+	else
+	{
+		if (OutError)
+		{
+			*OutError = FText::Format(LOCTEXT("Error_SaveManifest_SerializeFile", "Failed to serialize manifest '{0}'."), FText::FromString(InManifestFilePath));
+		}
+	}
+
+	if (LocFileNotifies.IsValid())
+	{
+		LocFileNotifies->PostFileWrite(InManifestFilePath);
+	}
+
+	return bSaved;
+}
+
+TSharedPtr<FInternationalizationArchive> FLocTextHelper::LoadArchiveImpl(const FString& InArchiveFilePath, const ELocTextHelperLoadFlags InLoadFlags, FText* OutError)
+{
+	TSharedPtr<FInternationalizationArchive> LocalArchive = MakeShareable(new FInternationalizationArchive());
+
+	// Attempt to load an existing archive first
+	if (!!(InLoadFlags & ELocTextHelperLoadFlags::Load))
+	{
+		bool bLoaded = false;
+		const bool bExisted = FPaths::FileExists(InArchiveFilePath);
+
+		if (bExisted)
+		{
+			if (LocFileNotifies.IsValid())
+			{
+				LocFileNotifies->PreFileRead(InArchiveFilePath);
+			}
+
+			TSharedPtr<FInternationalizationArchive> NativeArchive;
+			if (!NativeCulture.IsEmpty())
+			{
+				NativeArchive = Archives.FindRef(NativeCulture);
+			}
+			
+			if (FJsonInternationalizationArchiveSerializer::DeserializeArchiveFromFile(InArchiveFilePath, LocalArchive.ToSharedRef(), Manifest, NativeArchive))
+			{
+				bLoaded = true;
+			}
+			else
+			{
+				if (OutError)
+				{
+					*OutError = FText::Format(LOCTEXT("Error_LoadArchive_DeserializeFile", "Failed to deserialize archive '{0}'."), FText::FromString(InArchiveFilePath));
+				}
+			}
+
+			if (LocFileNotifies.IsValid())
+			{
+				LocFileNotifies->PostFileRead(InArchiveFilePath);
+			}
+		}
+
+		if (bLoaded)
+		{
+			return LocalArchive;
+		}
+		else if (bExisted)
+		{
+			// Don't allow fallback to Create if the file exists but could not be loaded
+			return nullptr;
+		}
+	}
+
+	// If we're allowed to create a manifest than we can never fail
+	if (!!(InLoadFlags & ELocTextHelperLoadFlags::Create))
+	{
+		return LocalArchive;
+	}
+
+	return nullptr;
+}
+
+bool FLocTextHelper::SaveArchiveImpl(const TSharedRef<const FInternationalizationArchive>& InArchive, const FString& InArchiveFilePath, FText* OutError) const
+{
+	bool bSaved = false;
+
+	if (LocFileNotifies.IsValid())
+	{
+		LocFileNotifies->PreFileWrite(InArchiveFilePath);
+	}
+
+	if (FJsonInternationalizationArchiveSerializer::SerializeArchiveToFile(InArchive, InArchiveFilePath))
+	{
+		bSaved = true;
+	}
+	else
+	{
+		if (OutError)
+		{
+			*OutError = FText::Format(LOCTEXT("Error_SaveArchive_SerializeFile", "Failed to serialize archive '{0}'."), FText::FromString(InArchiveFilePath));
+		}
+	}
+
+	if (LocFileNotifies.IsValid())
+	{
+		LocFileNotifies->PostFileWrite(InArchiveFilePath);
+	}
+
+	return bSaved;
+}
+
+TSharedPtr<FArchiveEntry> FLocTextHelper::FindTranslationImpl(const FString& InCulture, const FString& InNamespace, const FString& InKey, const TSharedPtr<FLocMetadataObject> InKeyMetadataObj) const
+{
+	TSharedPtr<FInternationalizationArchive> Archive = Archives.FindRef(InCulture);
+	checkf(Archive.IsValid(), TEXT("Attempted to find a translation, but no valid archive could be found for '%s'!"), *InCulture);
+	return Archive->FindEntryByKey(InNamespace, InKey, InKeyMetadataObj);
+}
+
+#undef LOCTEXT_NAMESPACE
diff --git a/Engine/Source/Runtime/Internationalization/Private/Serialization/JsonInternationalizationArchiveSerializer.cpp b/Engine/Source/Runtime/Internationalization/Private/Serialization/JsonInternationalizationArchiveSerializer.cpp
--- a/Engine/Source/Runtime/Internationalization/Private/Serialization/JsonInternationalizationArchiveSerializer.cpp
+++ b/Engine/Source/Runtime/Internationalization/Private/Serialization/JsonInternationalizationArchiveSerializer.cpp
@@ -5,6 +5,7 @@
 #include "Internationalization/InternationalizationMetadata.h"
 #include "JsonInternationalizationArchiveSerializer.h"
 #include "JsonInternationalizationMetadataSerializer.h"
+#include "LocTextHelper.h"
 
 
 DEFINE_LOG_CATEGORY_STATIC(LogInternationalizationArchiveSerializer, Log, All);
@@ -12,6 +13,7 @@ DEFINE_LOG_CATEGORY_STATIC(LogInternationalizationArchiveSerializer, Log, All);
 
 const FString FJsonInternationalizationArchiveSerializer::TAG_FORMATVERSION = TEXT("FormatVersion");
 const FString FJsonInternationalizationArchiveSerializer::TAG_NAMESPACE = TEXT("Namespace");
+const FString FJsonInternationalizationArchiveSerializer::TAG_KEY = TEXT("Key");
 const FString FJsonInternationalizationArchiveSerializer::TAG_CHILDREN = TEXT("Children");
 const FString FJsonInternationalizationArchiveSerializer::TAG_SUBNAMESPACES = TEXT("Subnamespaces");
 const FString FJsonInternationalizationArchiveSerializer::TAG_DEPRECATED_DEFAULTTEXT = TEXT("DefaultText");
@@ -21,6 +23,7 @@ const FString FJsonInternationalizationArchiveSerializer::TAG_SOURCE = TEXT("Sou
 const FString FJsonInternationalizationArchiveSerializer::TAG_SOURCE_TEXT = TEXT("Text");
 const FString FJsonInternationalizationArchiveSerializer::TAG_TRANSLATION = TEXT("Translation");
 const FString FJsonInternationalizationArchiveSerializer::TAG_TRANSLATION_TEXT = FJsonInternationalizationArchiveSerializer::TAG_SOURCE_TEXT;
+const FString FJsonInternationalizationArchiveSerializer::TAG_METADATA = TEXT("MetaData");
 const FString FJsonInternationalizationArchiveSerializer::TAG_METADATA_KEY = TEXT("Key");
 const FString FJsonInternationalizationArchiveSerializer::NAMESPACE_DELIMITER = TEXT(".");
 
@@ -59,169 +62,197 @@ struct FCompareStructuredArchiveEntryByNamespace
 };
 
 
-#if 0 // @todo Json: Serializing from FArchive is currently broken
-bool FJsonInternationalizationArchiveSerializer::DeserializeArchive( FArchive& Archive, TSharedRef< FInternationalizationArchive > InternationalizationArchive )
+bool FJsonInternationalizationArchiveSerializer::DeserializeArchive(const FString& InStr, TSharedRef<FInternationalizationArchive> InArchive, TSharedPtr<const FInternationalizationManifest> InManifest, TSharedPtr<const FInternationalizationArchive> InNativeArchive)
 {
 	TSharedPtr< FJsonObject > JsonArchiveObj;
-	TSharedRef< TJsonReader<> > Reader = TJsonReaderFactory<>::Create( &Archive );
-	bool bExecSuccessful = FJsonSerializer::Deserialize( Reader, JsonArchiveObj );
-
-	if (bExecSuccessful && JsonArchiveObj.IsValid())
+	TSharedRef< TJsonReader<> > Reader = TJsonReaderFactory<>::Create(InStr);
+	
+	if (FJsonSerializer::Deserialize(Reader, JsonArchiveObj) && JsonArchiveObj.IsValid())
 	{
-		bExecSuccessful = DeserializeInternal( JsonArchiveObj.ToSharedRef(), InternationalizationArchive );
-	}		
+		return DeserializeInternal(JsonArchiveObj.ToSharedRef(), InArchive, InManifest, InNativeArchive);
+	}
 
-	return bExecSuccessful;
+	return false;
 }
-#endif
 
 
-bool FJsonInternationalizationArchiveSerializer::DeserializeArchive( const FString& InStr, TSharedRef< FInternationalizationArchive > InternationalizationArchive )
+bool FJsonInternationalizationArchiveSerializer::DeserializeArchive(TSharedRef<FJsonObject> InJsonObj, TSharedRef<FInternationalizationArchive> InArchive, TSharedPtr<const FInternationalizationManifest> InManifest, TSharedPtr<const FInternationalizationArchive> InNativeArchive)
 {
-	TSharedPtr< FJsonObject > JsonArchiveObj;
-	TSharedRef< TJsonReader<> > Reader = TJsonReaderFactory<>::Create( InStr );
-	bool bExecSuccessful = FJsonSerializer::Deserialize( Reader, JsonArchiveObj );
+	return DeserializeInternal(InJsonObj, InArchive, InManifest, InNativeArchive);
+}
 
-	if ( bExecSuccessful && JsonArchiveObj.IsValid() )
+
+bool FJsonInternationalizationArchiveSerializer::DeserializeArchiveFromFile(const FString& InJsonFile, TSharedRef<FInternationalizationArchive> InArchive, TSharedPtr<const FInternationalizationManifest> InManifest, TSharedPtr<const FInternationalizationArchive> InNativeArchive)
+{
+	// Read in file as string
+	FString FileContents;
+	if (!FFileHelper::LoadFileToString(FileContents, *InJsonFile))
 	{
-		bExecSuccessful = DeserializeInternal( JsonArchiveObj.ToSharedRef(), InternationalizationArchive );
+		UE_LOG(LogInternationalizationArchiveSerializer, Error, TEXT("Failed to load archive '%s'."), *InJsonFile);
+		return false;
 	}
 
-	return bExecSuccessful;
-}
+	// Parse as JSON
+	TSharedPtr<FJsonObject> JsonObject;
 
+	TSharedRef<TJsonReader<>> JsonReader = TJsonReaderFactory<>::Create(FileContents);
+	if (!FJsonSerializer::Deserialize(JsonReader, JsonObject) || !JsonObject.IsValid())
+	{
+		UE_LOG(LogInternationalizationArchiveSerializer, Error, TEXT("Failed to parse archive '%s'. %s."), *InJsonFile, *JsonReader->GetErrorMessage());
+		return false;
+	}
 
-bool FJsonInternationalizationArchiveSerializer::DeserializeArchive( TSharedRef< FJsonObject > InJsonObj, TSharedRef< FInternationalizationArchive > InternationalizationArchive )
-{
-	return DeserializeInternal( InJsonObj, InternationalizationArchive );
+	return DeserializeInternal(JsonObject.ToSharedRef(), InArchive, InManifest, InNativeArchive);
 }
 
 
-#if 0 // @todo Json: Serializing from FArchive is currently broken
-bool FJsonInternationalizationArchiveSerializer::SerializeArchive( TSharedRef< const FInternationalizationArchive > InternationalizationArchive, FArchive& Archive )
+bool FJsonInternationalizationArchiveSerializer::SerializeArchive( TSharedRef< const FInternationalizationArchive > InArchive, FString& Str )
 {
 	TSharedRef< FJsonObject > JsonArchiveObj = MakeShareable( new FJsonObject );
-	bool bExecSuccessful = SerializeInternal( InternationalizationArchive, JsonArchiveObj );
+	bool bExecSuccessful = SerializeInternal( InArchive, JsonArchiveObj );
 
 	if( bExecSuccessful )
 	{
-		TSharedRef< TJsonWriter<> > Writer = TJsonWriterFactory<>::Create( &Archive );
+		TSharedRef< TJsonWriter<> > Writer = TJsonWriterFactory<>::Create( &Str );
 		bExecSuccessful = FJsonSerializer::Serialize( JsonArchiveObj, Writer );
 		Writer->Close();
 	}
+
 	return bExecSuccessful;
 }
-#endif
 
 
-bool FJsonInternationalizationArchiveSerializer::SerializeArchive( TSharedRef< const FInternationalizationArchive > InternationalizationArchive, FString& Str )
+bool FJsonInternationalizationArchiveSerializer::SerializeArchive( TSharedRef< const FInternationalizationArchive > InArchive, TSharedRef< FJsonObject > InJsonObj )
 {
-	TSharedRef< FJsonObject > JsonArchiveObj = MakeShareable( new FJsonObject );
-	bool bExecSuccessful = SerializeInternal( InternationalizationArchive, JsonArchiveObj );
+	return SerializeInternal( InArchive, InJsonObj );
+}
 
-	if( bExecSuccessful )
+
+bool FJsonInternationalizationArchiveSerializer::SerializeArchiveToFile(TSharedRef<const FInternationalizationArchive> InArchive, const FString& InJsonFile)
+{
+	TSharedRef<FJsonObject> JsonObject = MakeShareable(new FJsonObject());
+	if (!SerializeArchive(InArchive, JsonObject))
 	{
-		TSharedRef< TJsonWriter<> > Writer = TJsonWriterFactory<>::Create( &Str );
-		bExecSuccessful = FJsonSerializer::Serialize( JsonArchiveObj, Writer );
-		Writer->Close();
+		UE_LOG(LogInternationalizationArchiveSerializer, Error, TEXT("Failed to serialize archive '%s'."), *InJsonFile);
+		return false;
 	}
-	return bExecSuccessful;
-}
 
+	// Print the JSON data to a string
+	FString OutputJsonString;
+	TSharedRef<TJsonWriter<>> JsonWriter = TJsonWriterFactory<>::Create(&OutputJsonString);
+	FJsonSerializer::Serialize(JsonObject, JsonWriter);
 
-bool FJsonInternationalizationArchiveSerializer::SerializeArchive( TSharedRef< const FInternationalizationArchive > InternationalizationArchive, TSharedRef< FJsonObject > JsonObj )
-{
-	return SerializeInternal( InternationalizationArchive, JsonObj );
+	// Save the JSON string (force Unicode for our manifest and archive files)
+	if (!FFileHelper::SaveStringToFile(OutputJsonString, *InJsonFile, FFileHelper::EEncodingOptions::ForceUnicode))
+	{
+		UE_LOG(LogInternationalizationArchiveSerializer, Error, TEXT("Failed to save archive '%s'."), *InJsonFile);
+		return false;
+	}
+
+	return true;
 }
 
 
-bool FJsonInternationalizationArchiveSerializer::DeserializeInternal( TSharedRef< FJsonObject > InJsonObj, TSharedRef< FInternationalizationArchive > InternationalizationArchive )
+bool FJsonInternationalizationArchiveSerializer::DeserializeInternal(TSharedRef<FJsonObject> InJsonObj, TSharedRef<FInternationalizationArchive> InArchive, TSharedPtr<const FInternationalizationManifest> InManifest, TSharedPtr<const FInternationalizationArchive> InNativeArchive)
 {
-	if( InJsonObj->HasField( TAG_FORMATVERSION ) )
+	if (InJsonObj->HasField(TAG_FORMATVERSION))
 	{
-		const int32 FormatVersion = static_cast<int>(InJsonObj->GetNumberField( TAG_FORMATVERSION ));
-		InternationalizationArchive->SetFormatVersion(static_cast<FInternationalizationArchive::EFormatVersion>(FormatVersion));
+		const int32 FormatVersion = static_cast<int32>(InJsonObj->GetNumberField(TAG_FORMATVERSION));
+		InArchive->SetFormatVersion(static_cast<FInternationalizationArchive::EFormatVersion>(FormatVersion));
 	}
 	else
 	{
-		InternationalizationArchive->SetFormatVersion(FInternationalizationArchive::EFormatVersion::Initial);
+		InArchive->SetFormatVersion(FInternationalizationArchive::EFormatVersion::Initial);
 	}
 
-	return JsonObjToArchive( InJsonObj, TEXT(""), InternationalizationArchive );
+	if (InArchive->GetFormatVersion() < FInternationalizationArchive::EFormatVersion::AddedKeys && !InManifest.IsValid())
+	{
+		// Cannot load these archives without a manifest to key against
+		return false;
+	}
+
+	if (JsonObjToArchive(InJsonObj, FString(), InArchive, InManifest, InNativeArchive))
+	{
+		// We've been upgraded to the latest format now...
+		InArchive->SetFormatVersion(FInternationalizationArchive::EFormatVersion::Latest);
+		return true;
+	}
+
+	return false;
 }
 
 
-bool FJsonInternationalizationArchiveSerializer::SerializeInternal( TSharedRef< const FInternationalizationArchive > InInternationalizationArchive, TSharedRef< FJsonObject > JsonObj )
+bool FJsonInternationalizationArchiveSerializer::SerializeInternal(TSharedRef<const FInternationalizationArchive> InArchive, TSharedRef<FJsonObject> JsonObj)
 {
-	TSharedPtr< FStructuredArchiveEntry > RootElement = MakeShareable( new FStructuredArchiveEntry( TEXT("") ) );
+	TSharedPtr<FStructuredArchiveEntry> RootElement = MakeShareable(new FStructuredArchiveEntry(FString()));
 
 	// Condition the data so that it exists in a structured hierarchy for easy population of the JSON object.
-	GenerateStructuredData( InInternationalizationArchive, RootElement );
+	GenerateStructuredData(InArchive, RootElement);
 
-	SortStructuredData( RootElement );
+	SortStructuredData(RootElement);
 
 	// Clear anything that may be in the JSON object
 	JsonObj->Values.Empty();
 
 	// Set format version.
-	JsonObj->SetNumberField(TAG_FORMATVERSION, static_cast<double>(InInternationalizationArchive->GetFormatVersion()));
+	JsonObj->SetNumberField(TAG_FORMATVERSION, static_cast<double>(InArchive->GetFormatVersion()));
 
 	// Setup the JSON object using the structured data created
-	StructuredDataToJsonObj( RootElement, JsonObj );
+	StructuredDataToJsonObj(RootElement, JsonObj);
 	return true;
 }
 
 
-bool FJsonInternationalizationArchiveSerializer::JsonObjToArchive( TSharedRef< FJsonObject > InJsonObj, FString ParentNamespace, TSharedRef< FInternationalizationArchive > InternationalizationArchive )
+bool FJsonInternationalizationArchiveSerializer::JsonObjToArchive(TSharedRef<FJsonObject> InJsonObj, const FString& ParentNamespace, TSharedRef<FInternationalizationArchive> InArchive, TSharedPtr<const FInternationalizationManifest> InManifest, TSharedPtr<const FInternationalizationArchive> InNativeArchive)
 {
 	bool bConvertSuccess = true;
 	FString AccumulatedNamespace = ParentNamespace;
-	
-	if( InJsonObj->HasField( TAG_NAMESPACE) )
+
+	if (InJsonObj->HasField(TAG_NAMESPACE))
 	{
-		if( !( AccumulatedNamespace.IsEmpty() ) )
+		if (!(AccumulatedNamespace.IsEmpty()))
 		{
 			AccumulatedNamespace += NAMESPACE_DELIMITER;
 		}
-		AccumulatedNamespace += InJsonObj->GetStringField( TAG_NAMESPACE );
+		AccumulatedNamespace += InJsonObj->GetStringField(TAG_NAMESPACE);
 	}
 	else
 	{
-		UE_LOG( LogInternationalizationArchiveSerializer, Warning,TEXT("Encountered an object with a missing namespace while converting to Internationalization archive.") );
+		UE_LOG(LogInternationalizationArchiveSerializer, Warning, TEXT("Encountered an object with a missing namespace while converting to Internationalization archive."));
 		bConvertSuccess = false;
 	}
 
 	// Process all the child objects
-	if( bConvertSuccess && InJsonObj->HasField( TAG_CHILDREN ) )
+	if (bConvertSuccess && InJsonObj->HasField(TAG_CHILDREN))
 	{
-		const TArray< TSharedPtr<FJsonValue> > ChildrenArray = InJsonObj->GetArrayField( TAG_CHILDREN );
+		const TArray< TSharedPtr<FJsonValue> > ChildrenArray = InJsonObj->GetArrayField(TAG_CHILDREN);
 
-		for( TArray< TSharedPtr< FJsonValue > >::TConstIterator ChildIter( ChildrenArray.CreateConstIterator() ); ChildIter; ++ChildIter )
+		for (TArray< TSharedPtr< FJsonValue > >::TConstIterator ChildIter(ChildrenArray.CreateConstIterator()); ChildIter; ++ChildIter)
 		{
 			const TSharedPtr< FJsonValue >  ChildEntry = *ChildIter;
 			const TSharedPtr< FJsonObject > ChildJSONObject = ChildEntry->AsObject();
 
 			FString SourceText;
 			TSharedPtr< FLocMetadataObject > SourceMetadata;
-			if( ChildJSONObject->HasTypedField< EJson::String >( TAG_DEPRECATED_DEFAULTTEXT ) )
+			if (ChildJSONObject->HasTypedField< EJson::String >(TAG_DEPRECATED_DEFAULTTEXT))
 			{
-				SourceText = ChildJSONObject->GetStringField( TAG_DEPRECATED_DEFAULTTEXT );
-			} 
-			else if( ChildJSONObject->HasTypedField< EJson::Object >( TAG_SOURCE ) )
+				SourceText = ChildJSONObject->GetStringField(TAG_DEPRECATED_DEFAULTTEXT);
+			}
+			else if (ChildJSONObject->HasTypedField< EJson::Object >(TAG_SOURCE))
 			{
-				const TSharedPtr< FJsonObject > SourceJSONObject = ChildJSONObject->GetObjectField( TAG_SOURCE );
-				if( SourceJSONObject->HasTypedField< EJson::String >( TAG_SOURCE_TEXT ) )
+				const TSharedPtr< FJsonObject > SourceJSONObject = ChildJSONObject->GetObjectField(TAG_SOURCE);
+				if (SourceJSONObject->HasTypedField< EJson::String >(TAG_SOURCE_TEXT))
 				{
-					SourceText = SourceJSONObject->GetStringField( TAG_SOURCE_TEXT );
+					SourceText = SourceJSONObject->GetStringField(TAG_SOURCE_TEXT);
 
 					// Source meta data is mixed in with the source text, we'll process metadata if the source json object has more than one entry
-					if( SourceJSONObject->Values.Num() > 1 )
+					if (SourceJSONObject->Values.Num() > 1)
 					{
 						// We load in the entire source object as metadata and just remove the source text.
-						FJsonInternationalizationMetaDataSerializer::DeserializeMetadata( SourceJSONObject.ToSharedRef(), SourceMetadata );
-						if( SourceMetadata.IsValid() )
+						FJsonInternationalizationMetaDataSerializer::DeserializeMetadata(SourceJSONObject.ToSharedRef(), SourceMetadata);
+						if (SourceMetadata.IsValid())
 						{
-							SourceMetadata->Values.Remove( TAG_SOURCE_TEXT );
+							SourceMetadata->Values.Remove(TAG_SOURCE_TEXT);
 						}
 					}
 				}
@@ -237,25 +268,25 @@ bool FJsonInternationalizationArchiveSerializer::JsonObjToArchive( TSharedRef< F
 
 			FString TranslationText;
 			TSharedPtr< FLocMetadataObject > TranslationMetadata;
-			if( ChildJSONObject->HasTypedField< EJson::String >( TAG_DEPRECATED_TRANSLATEDTEXT ) )
+			if (ChildJSONObject->HasTypedField< EJson::String >(TAG_DEPRECATED_TRANSLATEDTEXT))
 			{
-				TranslationText = ChildJSONObject->GetStringField( TAG_DEPRECATED_TRANSLATEDTEXT );
-			} 
-			else if( ChildJSONObject->HasTypedField< EJson::Object >( TAG_TRANSLATION ) )
+				TranslationText = ChildJSONObject->GetStringField(TAG_DEPRECATED_TRANSLATEDTEXT);
+			}
+			else if (ChildJSONObject->HasTypedField< EJson::Object >(TAG_TRANSLATION))
 			{
-				const TSharedPtr< FJsonObject > TranslationJSONObject = ChildJSONObject->GetObjectField( TAG_TRANSLATION );
-				if( TranslationJSONObject->HasTypedField< EJson::String >( TAG_TRANSLATION_TEXT ) )
+				const TSharedPtr< FJsonObject > TranslationJSONObject = ChildJSONObject->GetObjectField(TAG_TRANSLATION);
+				if (TranslationJSONObject->HasTypedField< EJson::String >(TAG_TRANSLATION_TEXT))
 				{
-					TranslationText = TranslationJSONObject->GetStringField( TAG_TRANSLATION_TEXT );
+					TranslationText = TranslationJSONObject->GetStringField(TAG_TRANSLATION_TEXT);
 
 					// Source meta data is mixed in with the source text, we'll process metadata if the source json object has more than one entry
-					if( TranslationJSONObject->Values.Num() > 1 )
+					if (TranslationJSONObject->Values.Num() > 1)
 					{
 						// We load in the entire source object as metadata and remove the source text
-						FJsonInternationalizationMetaDataSerializer::DeserializeMetadata( TranslationJSONObject.ToSharedRef(), TranslationMetadata );
-						if( TranslationJSONObject.IsValid() )
+						FJsonInternationalizationMetaDataSerializer::DeserializeMetadata(TranslationJSONObject.ToSharedRef(), TranslationMetadata);
+						if (TranslationJSONObject.IsValid())
 						{
-							TranslationJSONObject->Values.Remove( TAG_TRANSLATION_TEXT );
+							TranslationJSONObject->Values.Remove(TAG_TRANSLATION_TEXT);
 						}
 					}
 				}
@@ -269,46 +300,78 @@ bool FJsonInternationalizationArchiveSerializer::JsonObjToArchive( TSharedRef< F
 				bConvertSuccess = false;
 			}
 
-			if( bConvertSuccess )
+			if (bConvertSuccess)
 			{
-				FLocItem Source( SourceText );
+				FLocItem Source(SourceText);
 				Source.MetadataObj = SourceMetadata;
 
-				FLocItem Translation( TranslationText );
+				FLocItem Translation(TranslationText);
 				Translation.MetadataObj = TranslationMetadata;
 
 				bool bIsOptional = false;
-				if( ChildJSONObject->HasTypedField< EJson::Boolean >( TAG_OPTIONAL ) )
+				if (ChildJSONObject->HasTypedField< EJson::Boolean >(TAG_OPTIONAL))
 				{
-					bIsOptional = ChildJSONObject->GetBoolField( TAG_OPTIONAL );
+					bIsOptional = ChildJSONObject->GetBoolField(TAG_OPTIONAL);
 				}
 
-				TSharedPtr< FLocMetadataObject > MetadataNode;
-				if( ChildJSONObject->HasTypedField< EJson::Object >( TAG_METADATA_KEY ) )
+				TArray<FString> Keys;
+				TSharedPtr< FLocMetadataObject > KeyMetadataNode;
+				if (InArchive->GetFormatVersion() < FInternationalizationArchive::EFormatVersion::AddedKeys)
+				{
+					// We used to store the key meta-data as a top-level value, rather than within a "MetaData" object
+					if (ChildJSONObject->HasTypedField< EJson::Object >(TAG_METADATA_KEY))
+					{
+						const TSharedPtr< FJsonObject > MetaDataKeyJSONObject = ChildJSONObject->GetObjectField(TAG_METADATA_KEY);
+						FJsonInternationalizationMetaDataSerializer::DeserializeMetadata(MetaDataKeyJSONObject.ToSharedRef(), KeyMetadataNode);
+					}
+
+					if (InManifest.IsValid())
+					{
+						// We have no key in the archive data, so we must try and infer it from the manifest
+						FLocTextHelper::FindKeysForLegacyTranslation(InManifest.ToSharedRef(), InNativeArchive, AccumulatedNamespace, SourceText, KeyMetadataNode, Keys);
+					}
+				}
+				else
 				{
-					const TSharedPtr< FJsonObject > MetaDataKeyJSONObject = ChildJSONObject->GetObjectField( TAG_METADATA_KEY );
-					FJsonInternationalizationMetaDataSerializer::DeserializeMetadata( MetaDataKeyJSONObject.ToSharedRef(), MetadataNode );
+					if (ChildJSONObject->HasTypedField< EJson::String >(TAG_KEY))
+					{
+						Keys.Add(ChildJSONObject->GetStringField(TAG_KEY));
+					}
+
+					if (ChildJSONObject->HasTypedField< EJson::Object >(TAG_METADATA))
+					{
+						const TSharedPtr< FJsonObject > MetaDataJSONObject = ChildJSONObject->GetObjectField(TAG_METADATA);
+
+						if (MetaDataJSONObject->HasTypedField< EJson::Object >(TAG_METADATA_KEY))
+						{
+							const TSharedPtr< FJsonObject > MetaDataKeyJSONObject = MetaDataJSONObject->GetObjectField(TAG_METADATA_KEY);
+							FJsonInternationalizationMetaDataSerializer::DeserializeMetadata(MetaDataKeyJSONObject.ToSharedRef(), KeyMetadataNode);
+						}
+					}
 				}
 
-				bool bAddSuccessful = InternationalizationArchive->AddEntry( AccumulatedNamespace, Source, Translation, MetadataNode, bIsOptional );
-				if( !bAddSuccessful )
+				for (const FString& Key : Keys)
 				{
-					UE_LOG( LogInternationalizationArchiveSerializer, Warning,TEXT("Could not add JSON entry to the Internationalization archive: Namespace:%s DefaultText:%s"), *AccumulatedNamespace, *SourceText );
+					const bool bAddSuccessful = InArchive->AddEntry(AccumulatedNamespace, Key, Source, Translation, KeyMetadataNode, bIsOptional);
+					if (!bAddSuccessful)
+					{
+						UE_LOG(LogInternationalizationArchiveSerializer, Warning, TEXT("Could not add JSON entry to the Internationalization archive: Namespace:%s Key:%s DefaultText:%s"), *AccumulatedNamespace, *Key, *SourceText);
+					}
 				}
 			}
 		}
 	}
 
-	if( bConvertSuccess && InJsonObj->HasField( TAG_SUBNAMESPACES ) )
+	if (bConvertSuccess && InJsonObj->HasField(TAG_SUBNAMESPACES))
 	{
-		const TArray< TSharedPtr<FJsonValue> > SubnamespaceArray = InJsonObj->GetArrayField( TAG_SUBNAMESPACES );
+		const TArray< TSharedPtr<FJsonValue> > SubnamespaceArray = InJsonObj->GetArrayField(TAG_SUBNAMESPACES);
 
-		for(TArray< TSharedPtr< FJsonValue > >::TConstIterator SubnamespaceIter( SubnamespaceArray.CreateConstIterator() ); SubnamespaceIter; ++SubnamespaceIter )
+		for (TArray< TSharedPtr< FJsonValue > >::TConstIterator SubnamespaceIter(SubnamespaceArray.CreateConstIterator()); SubnamespaceIter; ++SubnamespaceIter)
 		{
 			const TSharedPtr< FJsonValue >  SubnamespaceEntry = *SubnamespaceIter;
 			const TSharedPtr< FJsonObject > SubnamespaceJSONObject = SubnamespaceEntry->AsObject();
 
-			if( !JsonObjToArchive( SubnamespaceJSONObject.ToSharedRef(), AccumulatedNamespace, InternationalizationArchive ) )
+			if (!JsonObjToArchive(SubnamespaceJSONObject.ToSharedRef(), AccumulatedNamespace, InArchive, InManifest, InNativeArchive))
 			{
 				bConvertSuccess = false;
 				break;
@@ -320,10 +383,10 @@ bool FJsonInternationalizationArchiveSerializer::JsonObjToArchive( TSharedRef< F
 }
 
 
-void FJsonInternationalizationArchiveSerializer::GenerateStructuredData( TSharedRef< const FInternationalizationArchive > InInternationalizationArchive, TSharedPtr<FStructuredArchiveEntry> RootElement )
+void FJsonInternationalizationArchiveSerializer::GenerateStructuredData( TSharedRef< const FInternationalizationArchive > InArchive, TSharedPtr<FStructuredArchiveEntry> RootElement )
 {
 	//Loop through all the unstructured archive entries and build up our structured hierarchy
-	for(TArchiveEntryContainer::TConstIterator It( InInternationalizationArchive->GetEntryIterator() ); It; ++It)
+	for(FArchiveEntryByStringContainer::TConstIterator It( InArchive->GetEntriesBySourceTextIterator() ); It; ++It)
 	{
 		const TSharedRef< FArchiveEntry > UnstructuredArchiveEntry = It.Value();
 
@@ -428,13 +491,18 @@ void FJsonInternationalizationArchiveSerializer::StructuredDataToJsonObj( TShare
 		TranslationNode->SetStringField( TAG_TRANSLATION_TEXT, ProcessedTranslation );
 		EntryNode->SetObjectField( TAG_TRANSLATION, TranslationNode );
 
+		EntryNode->SetStringField( TAG_KEY, Entry->Key );
+
 		if( Entry->KeyMetadataObj.IsValid() )
 		{
-			TSharedPtr< FJsonObject > KeyDataNode;
-			FJsonInternationalizationMetaDataSerializer::SerializeMetadata( Entry->KeyMetadataObj.ToSharedRef(), KeyDataNode );
-			if( KeyDataNode.IsValid() )
+			TSharedRef< FJsonObject > MetaDataNode = MakeShareable(new FJsonObject());
+			EntryNode->SetObjectField( TAG_METADATA, MetaDataNode );
+
+			TSharedPtr< FJsonObject > KeyMetaDataNode;
+			FJsonInternationalizationMetaDataSerializer::SerializeMetadata( Entry->KeyMetadataObj.ToSharedRef(), KeyMetaDataNode );
+			if( KeyMetaDataNode.IsValid() )
 			{
-				EntryNode->SetObjectField( TAG_METADATA_KEY, KeyDataNode );
+				MetaDataNode->SetObjectField( TAG_METADATA_KEY, KeyMetaDataNode );
 			}
 		}
 
diff --git a/Engine/Source/Runtime/Internationalization/Private/Serialization/JsonInternationalizationManifestSerializer.cpp b/Engine/Source/Runtime/Internationalization/Private/Serialization/JsonInternationalizationManifestSerializer.cpp
--- a/Engine/Source/Runtime/Internationalization/Private/Serialization/JsonInternationalizationManifestSerializer.cpp
+++ b/Engine/Source/Runtime/Internationalization/Private/Serialization/JsonInternationalizationManifestSerializer.cpp
@@ -46,23 +46,6 @@ struct FCompareStructuredEntryByNamespace
 };
 
 
-#if 0 // @todo Json: Serializing from FArchive is currently broken
-bool FJsonInternationalizationManifestSerializer::DeserializeManifest( FArchive& Archive, TSharedRef< FInternationalizationManifest > Manifest )
-{
-	TSharedPtr< FJsonObject > JsonManifestObj;
-	TSharedRef< TJsonReader<> > Reader = TJsonReaderFactory<>::Create( &Archive );
-	bool bExecSuccessful = FJsonSerializer::Deserialize( Reader, JsonManifestObj );
-	
-	if ( bExecSuccessful && JsonManifestObj.IsValid() )
-	{
-		bExecSuccessful = DeserializeInternal( JsonManifestObj.ToSharedRef(), Manifest );
-	}		
-
-	return bExecSuccessful;
-}
-#endif
-
-
 bool FJsonInternationalizationManifestSerializer::DeserializeManifest( const FString& InStr, TSharedRef< FInternationalizationManifest > Manifest )
 {
 	TSharedPtr<FJsonObject> JsonManifestObj;
@@ -84,21 +67,28 @@ bool FJsonInternationalizationManifestSerializer::DeserializeManifest( TSharedRe
 }
 
 
-#if 0 // @todo Json: Serializing from FArchive is currently broken
-bool FJsonInternationalizationManifestSerializer::SerializeManifest( TSharedRef< const FInternationalizationManifest > Manifest, FArchive& Archive )
+bool FJsonInternationalizationManifestSerializer::DeserializeManifestFromFile( const FString& InJsonFile, TSharedRef< FInternationalizationManifest > Manifest )
 {
-	TSharedRef< FJsonObject > JsonManifestObj = MakeShareable( new FJsonObject );
-	bool bExecSuccessful = SerializeInternal( Manifest, JsonManifestObj );
+	// Read in file as string
+	FString FileContents;
+	if (!FFileHelper::LoadFileToString(FileContents, *InJsonFile))
+	{
+		UE_LOG(LogInternationalizationManifestSerializer, Error, TEXT("Failed to load manifest '%s'."), *InJsonFile);
+		return false;
+	}
+
+	// Parse as JSON
+	TSharedPtr<FJsonObject> JsonObject;
 
-	if( bExecSuccessful )
+	TSharedRef<TJsonReader<>> JsonReader = TJsonReaderFactory<>::Create(FileContents);
+	if (!FJsonSerializer::Deserialize(JsonReader, JsonObject) || !JsonObject.IsValid())
 	{
-		TSharedRef< TJsonWriter<> > Writer = TJsonWriterFactory<>::Create( &Archive );
-		bExecSuccessful = FJsonSerializer::Serialize( JsonManifestObj, Writer );
-		Writer->Close();
+		UE_LOG(LogInternationalizationManifestSerializer, Error, TEXT("Failed to parse manifest '%s'. %s."), *InJsonFile, *JsonReader->GetErrorMessage());
+		return false;
 	}
-	return bExecSuccessful;
+
+	return DeserializeInternal(JsonObject.ToSharedRef(), Manifest);
 }
-#endif
 
 
 bool FJsonInternationalizationManifestSerializer::SerializeManifest( TSharedRef< const FInternationalizationManifest > Manifest, FString& Str )
@@ -122,6 +112,31 @@ bool FJsonInternationalizationManifestSerializer::SerializeManifest( TSharedRef<
 }
 
 
+bool FJsonInternationalizationManifestSerializer::SerializeManifestToFile( TSharedRef< const FInternationalizationManifest > Manifest, const FString& InJsonFile )
+{
+	TSharedRef<FJsonObject> JsonObject = MakeShareable(new FJsonObject());
+	if (!SerializeManifest(Manifest, JsonObject))
+	{
+		UE_LOG(LogInternationalizationManifestSerializer, Error, TEXT("Failed to serialize manifest '%s'."), *InJsonFile);
+		return false;
+	}
+
+	// Print the JSON data to a string
+	FString OutputJsonString;
+	TSharedRef<TJsonWriter<>> JsonWriter = TJsonWriterFactory<>::Create(&OutputJsonString);
+	FJsonSerializer::Serialize(JsonObject, JsonWriter);
+
+	// Save the JSON string (force Unicode for our manifest and archive files)
+	if (!FFileHelper::SaveStringToFile(OutputJsonString, *InJsonFile, FFileHelper::EEncodingOptions::ForceUnicode))
+	{
+		UE_LOG(LogInternationalizationManifestSerializer, Error, TEXT("Failed to save manifest '%s'."), *InJsonFile);
+		return false;
+	}
+
+	return true;
+}
+
+
 bool FJsonInternationalizationManifestSerializer::DeserializeInternal( TSharedRef< FJsonObject > InJsonObj, TSharedRef< FInternationalizationManifest > Manifest )
 {
 	if( InJsonObj->HasField( TAG_FORMATVERSION ) )
@@ -331,7 +346,7 @@ bool FJsonInternationalizationManifestSerializer::JsonObjToManifest( TSharedRef<
 void FJsonInternationalizationManifestSerializer::GenerateStructuredData( TSharedRef< const FInternationalizationManifest > InManifest, TSharedPtr< FStructuredEntry > RootElement )
 {
 	//Loop through all the unstructured manifest entries and build up our structured hierarchy
-	for( TManifestEntryByContextIdContainer::TConstIterator It( InManifest->GetEntriesByContextIdIterator() ); It; ++It )
+	for( FManifestEntryByStringContainer::TConstIterator It( InManifest->GetEntriesByKeyIterator() ); It; ++It )
 	{
 		const TSharedRef< FManifestEntry > UnstructuredManifestEntry = It.Value();
 
diff --git a/Engine/Source/Runtime/Internationalization/Private/TextLocalizationResourceGenerator.cpp b/Engine/Source/Runtime/Internationalization/Private/TextLocalizationResourceGenerator.cpp
new file mode 100644
--- /dev/null
+++ b/Engine/Source/Runtime/Internationalization/Private/TextLocalizationResourceGenerator.cpp
@@ -0,0 +1,310 @@
+// Copyright 1998-2016 Epic Games, Inc. All Rights Reserved.
+
+#include "InternationalizationPrivatePCH.h"
+#include "TextLocalizationResourceGenerator.h"
+
+
+DEFINE_LOG_CATEGORY_STATIC(LogTextLocalizationResourceGenerator, Log, All);
+
+
+void FTextLocalizationResourceGenerator::FLocalizationEntryTracker::ReportCollisions() const
+{
+	for(auto i = Namespaces.CreateConstIterator(); i; ++i)
+	{
+		const FString& NamespaceName = i.Key();
+		const FLocalizationEntryTracker::FKeyTable& KeyTable = i.Value();
+		for(auto j = KeyTable.CreateConstIterator(); j; ++j)
+		{
+			const FString& KeyName = j.Key();
+			const FLocalizationEntryTracker::FEntryArray& EntryArray = j.Value();
+
+			bool bWasCollisionDetected = false;
+			for(int32 k = 0; k < EntryArray.Num(); ++k)
+			{
+				const FLocalizationEntryTracker::FEntry& LeftEntry = EntryArray[k];
+				for(int32 l = k + 1; l < EntryArray.Num(); ++l)
+				{
+					const FLocalizationEntryTracker::FEntry& RightEntry = EntryArray[l];
+					const bool bDoesLocalizedStringDiffer = !LeftEntry.LocalizedString.Equals( RightEntry.LocalizedString, ESearchCase::CaseSensitive );
+					bWasCollisionDetected = bDoesLocalizedStringDiffer;
+				}
+			}
+
+			if(bWasCollisionDetected)
+			{
+				FString CollidingEntryListString;
+				for(int32 k = 0; k < EntryArray.Num(); ++k)
+				{
+					const FLocalizationEntryTracker::FEntry& Entry = EntryArray[k];
+
+					if( !(CollidingEntryListString.IsEmpty()) )
+					{
+						CollidingEntryListString += '\n';
+					}
+
+					CollidingEntryListString += FString::Printf( TEXT("String: (%s)"), *Entry.LocalizedString );
+				}
+
+				UE_LOG(LogTextLocalizationResourceGenerator, Error, TEXT("Archives contain conflicting entries for (Namespace:%s, Key:%s):\n%s"), *NamespaceName, *KeyName, *CollidingEntryListString);
+			}
+		}
+	}
+}
+
+bool FTextLocalizationResourceGenerator::FLocalizationEntryTracker::WriteToArchive(FArchive* const Archive)
+{
+	bool WasSuccessful = true;
+
+	FLocalizationEntryTracker& NonConstThis = const_cast<FLocalizationEntryTracker&>(*this); // Necessary conversion because FArchive lacks const-correctness using the insertion operator.
+
+	if ( Archive == NULL )
+	{
+		return false;
+	}
+
+	Archive->SetForceUnicode(true);
+
+	// Write namespace count
+	uint32 NamespaceCount = NonConstThis.Namespaces.Num();
+	*Archive << NamespaceCount;
+
+	// Iterate through namespaces
+	for(auto NamespaceIterator = NonConstThis.Namespaces.CreateIterator(); NamespaceIterator; ++NamespaceIterator)
+	{
+		/*const*/ FString Namespace = NamespaceIterator.Key();
+		/*const*/ FLocalizationEntryTracker::FKeyTable& KeyTable = NamespaceIterator.Value();
+
+		// Write namespace.
+		*Archive << Namespace;
+
+		// Write key count.
+		const int32 KeyCountOffsetInFile = Archive->Tell();
+		const uint32 OriginalKeyCount = KeyTable.Num();
+		uint32 KeyCount = OriginalKeyCount;
+		*Archive << KeyCount;
+
+		// Iterate through keys and values
+		for(auto KeyIterator = KeyTable.CreateIterator(); KeyIterator; ++KeyIterator)
+		{
+			/*const*/ FString Key = KeyIterator.Key();
+			/*const*/ FLocalizationEntryTracker::FEntryArray& EntryArray = KeyIterator.Value();
+
+			// Skip this key if there are no entries.
+			if(EntryArray.Num() == 0)
+			{
+				UE_LOG(LogTextLocalizationResourceGenerator, Warning, TEXT("Archives contained no entries for key (%s)"), *(Key));
+				--KeyCount; // We've skipped an entry and thus a key and the key count must be adjusted.
+				continue;
+			}
+
+			// Find first valid entry.
+			/*const*/ FLocalizationEntryTracker::FEntry* Value = NULL;
+			for(int32 i = 0; i < EntryArray.Num(); ++i)
+			{
+				if( !( EntryArray[i].LocalizedString.IsEmpty() ) )
+				{
+					Value = &(EntryArray[i]);
+					break;
+				}
+			}
+
+			// Skip this key if there is no valid entry.
+			if( !(Value) )
+			{
+				UE_LOG(LogTextLocalizationResourceGenerator, Verbose, TEXT("Archives contained only blank entries for key (%s)"), *(Key));
+				--KeyCount; // We've skipped an entry and thus a key and the key count must be adjusted.
+				continue;
+			}
+
+			// Write key.
+			*Archive << Key;
+
+			// Write string entry.
+			*Archive << Value->SourceStringHash;
+
+			*Archive << Value->LocalizedString;
+		}
+
+		// If key count is different than original key count (due to skipped entries), adjust key count in file.
+		if(KeyCount != OriginalKeyCount)
+		{
+			const int32 LatestOffsetInFile = Archive->Tell();
+			Archive->Seek(KeyCountOffsetInFile);
+			*Archive << KeyCount;
+			Archive->Seek(LatestOffsetInFile);
+		}
+	}
+
+	return WasSuccessful;
+}
+
+bool FTextLocalizationResourceGenerator::Generate(const FLocTextHelper& InLocTextHelper, const FString& InCultureToGenerate, const bool bSkipSourceCheck, FArchive& InDestinationArchive)
+{
+	FLocalizationEntryTracker LocalizationEntryTracker;
+
+	// Add each manifest entry to the LocRes file
+	InLocTextHelper.EnumerateSourceTexts([&InLocTextHelper, &InCultureToGenerate, &LocalizationEntryTracker, &bSkipSourceCheck](TSharedRef<FManifestEntry> InManifestEntry) -> bool
+	{
+		// For each context, we may need to create a different or even multiple LocRes entries.
+		for (const FManifestContext& Context : InManifestEntry->Contexts)
+		{
+			// Find the correct translation based upon the native source text
+			FLocItem TranslationText;
+			InLocTextHelper.GetRuntimeText(InCultureToGenerate, InManifestEntry->Namespace, Context.Key, Context.KeyMetadataObj, ELocTextExportSourceMethod::NativeText, InManifestEntry->Source, TranslationText, bSkipSourceCheck);
+
+			// Add this entry to the LocRes
+			FLocalizationEntryTracker::FKeyTable& KeyTable = LocalizationEntryTracker.Namespaces.FindOrAdd(*InManifestEntry->Namespace);
+			FLocalizationEntryTracker::FEntryArray& EntryArray = KeyTable.FindOrAdd(*Context.Key);
+			FLocalizationEntryTracker::FEntry& NewEntry = EntryArray[EntryArray.AddDefaulted()];
+			NewEntry.SourceStringHash = FCrc::StrCrc32(*InManifestEntry->Source.Text);
+			NewEntry.LocalizedString = TranslationText.Text;
+		}
+
+		return true; // continue enumeration
+	}, true);
+
+	LocalizationEntryTracker.ReportCollisions();
+
+	// Write resource.
+	if (!LocalizationEntryTracker.WriteToArchive(&InDestinationArchive))
+	{
+		UE_LOG(LogTextLocalizationResourceGenerator, Error, TEXT("Failed to write localization entries to archive."));
+		return false;
+	}
+
+	return true;
+}
+
+bool FTextLocalizationResourceGenerator::GenerateAndUpdateLiveEntriesFromConfig(const FString& InConfigFilePath, const bool bSkipSourceCheck)
+{
+	FInternationalization& I18N = FInternationalization::Get();
+
+	const FString SectionName = TEXT("RegenerateResources");
+
+	// Get native culture.
+	FString NativeCulture;
+	if (!GConfig->GetString(*SectionName, TEXT("NativeCulture"), NativeCulture, InConfigFilePath))
+	{
+		UE_LOG(LogTextLocalizationResourceGenerator, Error, TEXT("No native culture specified."));
+		return false;
+	}
+
+	// Get source path.
+	FString SourcePath;
+	if (!GConfig->GetString(*SectionName, TEXT("SourcePath"), SourcePath, InConfigFilePath))
+	{
+		UE_LOG(LogTextLocalizationResourceGenerator, Error, TEXT("No source path specified."));
+		return false;
+	}
+
+	// Get destination path.
+	FString DestinationPath;
+	if (!GConfig->GetString(*SectionName, TEXT("DestinationPath"), DestinationPath, InConfigFilePath))
+	{
+		UE_LOG(LogTextLocalizationResourceGenerator, Error, TEXT("No destination path specified."));
+		return false;
+	}
+
+	// Get manifest name.
+	FString ManifestName;
+	if (!GConfig->GetString(*SectionName, TEXT("ManifestName"), ManifestName, InConfigFilePath))
+	{
+		UE_LOG(LogTextLocalizationResourceGenerator, Error, TEXT("No manifest name specified."));
+		return false;
+	}
+
+	// Get manifest name.
+	FString ArchiveName;
+	if (!GConfig->GetString(*SectionName, TEXT("ArchiveName"), ManifestName, InConfigFilePath))
+	{
+		UE_LOG(LogTextLocalizationResourceGenerator, Error, TEXT("No archive name specified."));
+		return false;
+	}
+
+	// Get resource name.
+	FString ResourceName;
+	if (!GConfig->GetString(*SectionName, TEXT("ResourceName"), ResourceName, InConfigFilePath))
+	{
+		UE_LOG(LogTextLocalizationResourceGenerator, Error, TEXT("No resource name specified."));
+		return false;
+	}
+
+	// Source path needs to be relative to Engine or Game directory
+	const FString ConfigFullPath = FPaths::ConvertRelativePathToFull(InConfigFilePath);
+	const FString EngineFullPath = FPaths::ConvertRelativePathToFull(FPaths::EngineConfigDir());
+	const bool IsEngineManifest = ConfigFullPath.StartsWith(EngineFullPath);
+
+	if (IsEngineManifest)
+	{
+		SourcePath = FPaths::Combine(*FPaths::EngineDir(), *SourcePath);
+		DestinationPath = FPaths::Combine(*FPaths::EngineDir(), *DestinationPath);
+	}
+	else
+	{
+		SourcePath = FPaths::Combine(*FPaths::GameDir(), *SourcePath);
+		DestinationPath = FPaths::Combine(*FPaths::GameDir(), *DestinationPath);
+	}
+
+	TArray<FString> CulturesToGenerate;
+	{
+		const FString CultureName = I18N.GetCurrentCulture()->GetName();
+		const TArray<FString> PrioritizedCultures = I18N.GetPrioritizedCultureNames(CultureName);
+		for (const FString& PrioritizedCulture : PrioritizedCultures)
+		{
+			if (FPaths::FileExists(SourcePath / PrioritizedCulture / ArchiveName))
+			{
+				CulturesToGenerate.Add(PrioritizedCulture);
+			}
+		}
+	}
+
+	if (CulturesToGenerate.Num() == 0)
+	{
+		UE_LOG(LogTextLocalizationResourceGenerator, Error, TEXT("No cultures to generate were specified."));
+		return false;
+	}
+
+	// Load the manifest and all archives
+	FLocTextHelper LocTextHelper(SourcePath, ManifestName, ArchiveName, NativeCulture, CulturesToGenerate, nullptr);
+	{
+		FText LoadError;
+		if (!LocTextHelper.LoadAll(ELocTextHelperLoadFlags::LoadOrCreate, &LoadError))
+		{
+			UE_LOG(LogTextLocalizationResourceGenerator, Error, TEXT("%s"), *LoadError.ToString());
+			return false;
+		}
+	}
+
+	TArray<TArray<uint8>> BackingBuffers;
+	BackingBuffers.SetNum(CulturesToGenerate.Num());
+
+	for (int32 CultureIndex = 0; CultureIndex < CulturesToGenerate.Num(); ++CultureIndex)
+	{
+		const FString& CultureName = CulturesToGenerate[CultureIndex];
+
+		TArray<uint8>& BackingBuffer = BackingBuffers[CultureIndex];
+		FMemoryWriter MemoryWriter(BackingBuffer, true);
+		const bool bDidGenerate = Generate(LocTextHelper, CultureName, bSkipSourceCheck, MemoryWriter);
+		MemoryWriter.Close();
+
+		if (!bDidGenerate)
+		{
+			UE_LOG(LogTextLocalizationResourceGenerator, Error, TEXT("Failed to generate localization resource for culture '%s'."), *CultureName);
+			return false;
+		}
+	}
+
+	for (int32 CultureIndex = 0; CultureIndex < CulturesToGenerate.Num(); ++CultureIndex)
+	{
+		const FString& CultureName = CulturesToGenerate[CultureIndex];
+		const FString CulturePath = DestinationPath / CultureName;
+		const FString ResourceFilePath = FPaths::ConvertRelativePathToFull(CulturePath / ResourceName);
+
+		TArray<uint8>& BackingBuffer = BackingBuffers[CultureIndex];
+		FMemoryReader MemoryReader(BackingBuffer, true);
+		FTextLocalizationManager::Get().UpdateFromLocalizationResource(MemoryReader, ResourceFilePath);
+		MemoryReader.Close();
+	}
+
+	return true;
+}
diff --git a/Engine/Source/Runtime/Internationalization/Public/Serialization/JsonInternationalizationArchiveSerializer.h b/Engine/Source/Runtime/Internationalization/Public/Serialization/JsonInternationalizationArchiveSerializer.h
--- a/Engine/Source/Runtime/Internationalization/Public/Serialization/JsonInternationalizationArchiveSerializer.h
+++ b/Engine/Source/Runtime/Internationalization/Public/Serialization/JsonInternationalizationArchiveSerializer.h
@@ -2,9 +2,9 @@
 
 #pragma once
 
-#include "IInternationalizationArchiveSerializer.h"
 #include "Json.h"
-
+#include "InternationalizationArchive.h"
+#include "InternationalizationManifest.h"
 
 /**
  * Used to arrange Internationalization archive data in a hierarchy based on namespace prior to json serialization.
@@ -23,92 +23,128 @@ public:
 
 
 /**
- * Implements a InternationalizationArchive serializer that serializes to and from Json encoded data.
+ * Implements a serializer that serializes to and from Json encoded data.
  */
 class INTERNATIONALIZATION_API FJsonInternationalizationArchiveSerializer
-	: public IInternationalizationArchiveSerializer
 {
 public:
 
-	/** Default constructor. */
-	FJsonInternationalizationArchiveSerializer() { }
+	/**
+	 * Deserializes an archive from a JSON string.
+	 *
+	 * @param InJsonObj			The JSON string to serialize from.
+	 * @param InArchive			The archive to populate from the JSON data.
+	 * @param InManifest		The manifest associated with the archive. May be null, but you won't be able to load archives with a version < FInternationalizationArchive::EFormatVersion::AddedKeys.
+	 * @param InNativeArchive	The native archive associated with the archive. May be null.
+	 *
+	 * @return true if deserialization was successful, false otherwise.
+	 */
+	static bool DeserializeArchive(const FString& InStr, TSharedRef<FInternationalizationArchive> InArchive, TSharedPtr<const FInternationalizationManifest> InManifest, TSharedPtr<const FInternationalizationArchive> InNativeArchive);
 
-public:
+	/**
+	 * Deserializes an archive from a JSON object.
+	 *
+	 * @param InJsonObj			The JSON object to serialize from.
+	 * @param InArchive			The archive to populate from the JSON data.
+	 * @param InManifest		The manifest associated with the archive. May be null, but you won't be able to load archives with a version < FInternationalizationArchive::EFormatVersion::AddedKeys.
+	 * @param InNativeArchive	The native archive associated with the archive. May be null.
+	 *
+	 * @return true if deserialization was successful, false otherwise.
+	 */
+	static bool DeserializeArchive(TSharedRef<FJsonObject> InJsonObj, TSharedRef<FInternationalizationArchive> InArchive, TSharedPtr<const FInternationalizationManifest> InManifest, TSharedPtr<const FInternationalizationArchive> InNativeArchive);
 
 	/**
-	 * Deserializes a Internationalization archive from a JSON object.
+	 * Deserializes an archive from a JSON file.
+	 *
+	 * @param InJsonFile		The path to the JSON file to serialize from.
+	 * @param InArchive			The archive to populate from the JSON data.
+	 * @param InManifest		The manifest associated with the archive. May be null, but you won't be able to load archives with a version < FInternationalizationArchive::EFormatVersion::AddedKeys.
+	 * @param InNativeArchive	The native archive associated with the archive. May be null.
 	 *
-	 * @param InJsonObj The JSON object to serialize from.
-	 * @param InternationalizationArchive The populated Internationalization archive.
 	 * @return true if deserialization was successful, false otherwise.
 	 */
-	bool DeserializeArchive( TSharedRef< FJsonObject > InJsonObj, TSharedRef< FInternationalizationArchive > InternationalizationArchive );
+	static bool DeserializeArchiveFromFile(const FString& InJsonFile, TSharedRef<FInternationalizationArchive> InArchive, TSharedPtr<const FInternationalizationManifest> InManifest, TSharedPtr<const FInternationalizationArchive> InNativeArchive);
 
 	/**
-	 * Serializes a Internationalization archive to a JSON object.
+	 * Serializes an archive to a JSON object.
+	 *
+	 * @param InArchive			The archive data to serialize.
+	 * @param InJsonObj			The JSON object to serialize into.
 	 *
-	 * @param InternationalizationArchive The Internationalization archive data to serialize.
-	 * @param JsonObj The JSON object to serialize into.
 	 * @return true if serialization was successful, false otherwise.
 	 */
-	bool SerializeArchive( TSharedRef< const FInternationalizationArchive > InternationalizationArchive, TSharedRef< FJsonObject > JsonObj );
-
-public:
+	static bool SerializeArchive(TSharedRef<const FInternationalizationArchive> InArchive, TSharedRef<FJsonObject> InJsonObj);
 
-	// IInternationalizationArchiveSerializer interface
-
-	virtual bool DeserializeArchive( const FString& InStr, TSharedRef< FInternationalizationArchive > InternationalizationArchive ) override;
-	virtual bool SerializeArchive( TSharedRef< const FInternationalizationArchive > InternationalizationArchive, FString& Str ) override;
+	/**
+	 * Serializes an archive to a JSON string.
+	 *
+	 * @param InArchive			The archive data to serialize.
+	 * @param Str				The string to fill with the JSON data.
+	 *
+	 * @return true if serialization was successful, false otherwise.
+	 */
+	static bool SerializeArchive(TSharedRef<const FInternationalizationArchive> InArchive, FString& Str);
 
-#if 0 // @todo Json: Serializing from FArchive is currently broken
-	virtual bool DeserializeArchive( FArchive& Archive, TSharedRef< FInternationalizationArchive > InternationalizationArchive ) override;
-	virtual bool SerializeArchive( TSharedRef< const FInternationalizationArchive > InternationalizationArchive, FArchive& Archive ) override;
-#endif
+	/**
+	* Serializes an archive to a JSON string.
+	*
+	* @param InArchive			The archive data to serialize.
+	* @param InJsonFile			The path to the JSON file to serialize to.
+	*
+	* @return true if serialization was successful, false otherwise.
+	*/
+	static bool SerializeArchiveToFile(TSharedRef<const FInternationalizationArchive> InArchive, const FString& InJsonFile);
 
 protected:
 
 	/**
-	 * Convert a JSON object to a Internationalization archive.
+	 * Deserializes an archive from a JSON object.
+	 *
+	 * @param InJsonObj			The JSON object to serialize from.
+	 * @param InArchive			The archive to populate from the JSON data.
+	 * @param InManifest		The manifest associated with the archive. May be null, but you won't be able to load archives with a version < FInternationalizationArchive::EFormatVersion::AddedKeys.
+	 * @param InNativeArchive	The native archive associated with the archive. May be null.
 	 *
-	 * @param InJsonObj The JSON object to serialize from.
-	 * @param InternationalizationArchive The Internationalization archive that will store the data.
 	 * @return true if deserialization was successful, false otherwise.
 	 */
-	bool DeserializeInternal( TSharedRef< FJsonObject > InJsonObj, TSharedRef< FInternationalizationArchive > InternationalizationArchive );
+	static bool DeserializeInternal(TSharedRef<FJsonObject> InJsonObj, TSharedRef<FInternationalizationArchive> InArchive, TSharedPtr<const FInternationalizationManifest> InManifest, TSharedPtr<const FInternationalizationArchive> InNativeArchive);
 
 	/**
 	 * Convert a Internationalization archive to a JSON object.
 	 *
-	 * @param InInternationalizationArchive The Internationalization archive object to serialize from.
+	 * @param InArchive The Internationalization archive object to serialize from.
 	 * @param JsonObj The Json object that will store the data.
 	 * @return true if serialization was successful, false otherwise.
 	 */
-	bool SerializeInternal( TSharedRef< const FInternationalizationArchive > InInternationalizationArchive, TSharedRef< FJsonObject > JsonObj );
+	static bool SerializeInternal(TSharedRef<const FInternationalizationArchive> InArchive, TSharedRef<FJsonObject> JsonObj);
 
 	/**
-	 * Recursive function that will traverse the JSON object and populate a Internationalization archive.
+	 * Recursive function that will traverse the JSON object and populate an archive.
+	 *
+	 * @param InJsonObj			The JSON object to serialize from.
+	 * @param ParentNamespace	The namespace of the parent JSON object.
+	 * @param InArchive			The archive to populate from the JSON data.
+	 * @param InManifest		The manifest associated with the archive. May be null, but you won't be able to load archives with a version < FInternationalizationArchive::EFormatVersion::AddedKeys.
+	 * @param InNativeArchive	The native archive associated with the archive. May be null.
 	 *
-	 * @param InJsonObj The JSON object.
-	 * @param ParentNamespace The namespace of the parent JSON object.
-	 * @param InternationalizationArchive The Internationalization archive that will store the data.
 	 * @return true if successful, false otherwise.
 	 */
-	bool JsonObjToArchive( TSharedRef< FJsonObject > InJsonObj, FString ParentNamespace, TSharedRef< FInternationalizationArchive > InternationalizationArchive );
+	static bool JsonObjToArchive(TSharedRef<FJsonObject> InJsonObj, const FString& ParentNamespace, TSharedRef<FInternationalizationArchive> InArchive, TSharedPtr<const FInternationalizationManifest> InManifest, TSharedPtr<const FInternationalizationArchive> InNativeArchive);
 
 	/**
 	 * Takes a Internationalization archive and arranges the data into a hierarchy based on namespace.
 	 *
-	 * @param InInternationalizationArchive The Internationalization archive.
+	 * @param InArchive The Internationalization archive.
 	 * @param RootElement The root element of the structured data.
 	 */
-	void GenerateStructuredData( TSharedRef< const FInternationalizationArchive > InInternationalizationArchive, TSharedPtr< FStructuredArchiveEntry > RootElement );
+	static void GenerateStructuredData( TSharedRef< const FInternationalizationArchive > InArchive, TSharedPtr< FStructuredArchiveEntry > RootElement );
 
 	/**
 	 * Goes through the structured, hierarchy based, archive data and does a non-culture specific sort on namespaces and default text.
 	 *
 	 * @param RootElement The root element of the structured data.
 	 */
-	void SortStructuredData( TSharedPtr< FStructuredArchiveEntry > InElement );
+	static void SortStructuredData( TSharedPtr< FStructuredArchiveEntry > InElement );
 
 	/**
 	 * Populates a JSON object from Internationalization archive data that has been structured based on namespace.
@@ -116,12 +152,13 @@ protected:
 	 * @param InElement Internationalization archive data structured based on namespace.
 	 * @param JsonObj JSON object to be populated.
 	 */
-	void StructuredDataToJsonObj(TSharedPtr< const FStructuredArchiveEntry > InElement, TSharedRef< FJsonObject > JsonObj );
+	static void StructuredDataToJsonObj(TSharedPtr< const FStructuredArchiveEntry > InElement, TSharedRef< FJsonObject > JsonObj );
 
 public:
 
 	static const FString TAG_FORMATVERSION;
 	static const FString TAG_NAMESPACE;
+	static const FString TAG_KEY;
 	static const FString TAG_CHILDREN;
 	static const FString TAG_SUBNAMESPACES;
 	static const FString TAG_DEPRECATED_DEFAULTTEXT;
@@ -131,6 +168,7 @@ public:
 	static const FString TAG_SOURCE_TEXT;
 	static const FString TAG_TRANSLATION;
 	static const FString TAG_TRANSLATION_TEXT;
+	static const FString TAG_METADATA;
 	static const FString TAG_METADATA_KEY;
 	static const FString NAMESPACE_DELIMITER;
 };
diff --git a/Engine/Source/Runtime/Internationalization/Public/Serialization/JsonInternationalizationManifestSerializer.h b/Engine/Source/Runtime/Internationalization/Public/Serialization/JsonInternationalizationManifestSerializer.h
--- a/Engine/Source/Runtime/Internationalization/Public/Serialization/JsonInternationalizationManifestSerializer.h
+++ b/Engine/Source/Runtime/Internationalization/Public/Serialization/JsonInternationalizationManifestSerializer.h
@@ -1,8 +1,8 @@
 // Copyright 1998-2016 Epic Games, Inc. All Rights Reserved.
-#pragma once
 
-#include "IInternationalizationManifestSerializer.h"
+#pragma once
 
+#include "InternationalizationManifest.h"
 
 struct FStructuredEntry
 {
@@ -23,14 +23,17 @@ class FJsonObject;
 
 
 class INTERNATIONALIZATION_API FJsonInternationalizationManifestSerializer
-	: public IInternationalizationManifestSerializer
 {
 public:
 
-	/** Default constructor. */
-	FJsonInternationalizationManifestSerializer() { }
-
-public:
+	/**
+	 * Deserializes a Internationalization manifest from a JSON string.
+	 *
+	 * @param InStr The JSON string to serialize from.
+	 * @param Manifest The populated Internationalization manifest.
+	 * @return true if deserialization was successful, false otherwise.
+	 */
+	static bool DeserializeManifest( const FString& InStr, TSharedRef< FInternationalizationManifest > Manifest );
 
 	/**
 	 * Deserializes a Internationalization manifest from a JSON object.
@@ -39,7 +42,25 @@ public:
 	 * @param Manifest The populated Internationalization manifest.
 	 * @return true if deserialization was successful, false otherwise.
 	 */
-	bool DeserializeManifest( TSharedRef< FJsonObject > InJsonObj, TSharedRef< FInternationalizationManifest > Manifest );
+	static bool DeserializeManifest( TSharedRef< FJsonObject > InJsonObj, TSharedRef< FInternationalizationManifest > Manifest );
+
+	/**
+	 * Deserializes a Internationalization manifest from a JSON file.
+	 *
+	 * @param InJsonFile The path to the JSON file to serialize from.
+	 * @param Manifest The populated Internationalization manifest.
+	 * @return true if deserialization was successful, false otherwise.
+	 */
+	static bool DeserializeManifestFromFile( const FString& InJsonFile, TSharedRef< FInternationalizationManifest > Manifest );
+
+	/**
+	 * Serializes a Internationalization manifest to a JSON string.
+	 *
+	 * @param Manifest The Internationalization manifest data to serialize.
+	 * @param Str The string to serialize into.
+	 * @return true if serialization was successful, false otherwise.
+	 */
+	static bool SerializeManifest( TSharedRef< const FInternationalizationManifest > Manifest, FString& Str );
 
 	/**
 	 * Serializes a Internationalization manifest to a JSON object.
@@ -48,19 +69,16 @@ public:
 	 * @param JsonObj The JSON object to serialize into.
 	 * @return true if serialization was successful, false otherwise.
 	 */
-	bool SerializeManifest( TSharedRef< const FInternationalizationManifest > Manifest, TSharedRef< FJsonObject > JsonObj );
-
-public:
+	static bool SerializeManifest( TSharedRef< const FInternationalizationManifest > Manifest, TSharedRef< FJsonObject > JsonObj );
 
-	// IInternationalizationManifestSerializer interface
-
-	virtual bool DeserializeManifest( const FString& InStr, TSharedRef< FInternationalizationManifest > Manifest ) override;
-	virtual bool SerializeManifest( TSharedRef< const FInternationalizationManifest > Manifest, FString& Str ) override;
-
-#if 0 // @todo Json: Serializing from FArchive is currently broken
-	virtual bool DeserializeManifest( FArchive& Archive, TSharedRef< FInternationalizationManifest > Manifest ) override;
-	virtual bool SerializeManifest(TSharedRef< const FInternationalizationManifest > Manifest, FArchive& Archive ) override;
-#endif
+	/**
+	 * Serializes a Internationalization manifest to a JSON file.
+	 *
+	 * @param Manifest The Internationalization manifest data to serialize.
+	 * @param InJsonFile The path to the JSON file to serialize to.
+	 * @return true if serialization was successful, false otherwise.
+	 */
+	static bool SerializeManifestToFile( TSharedRef< const FInternationalizationManifest > Manifest, const FString& InJsonFile );
 
 protected:
 
@@ -72,7 +90,7 @@ protected:
 	 *
 	 * @return true if deserialization was successful, false otherwise.
 	 */
-	bool DeserializeInternal( const TSharedRef< FJsonObject > InJsonObj, TSharedRef< FInternationalizationManifest > Manifest );
+	static bool DeserializeInternal( const TSharedRef< FJsonObject > InJsonObj, TSharedRef< FInternationalizationManifest > Manifest );
 
 	/**
 	 * Convert a Internationalization manifest to a JSON object.
@@ -82,7 +100,7 @@ protected:
 	 *
 	 * @return true if serialization was successful, false otherwise.
 	 */
-	bool SerializeInternal( TSharedRef< const FInternationalizationManifest > InManifest, TSharedRef< FJsonObject > JsonObj );
+	static bool SerializeInternal( TSharedRef< const FInternationalizationManifest > InManifest, TSharedRef< FJsonObject > JsonObj );
 
 	/**
 	 * Recursive function that will traverse the JSON object and populate a Internationalization manifest.
@@ -93,7 +111,7 @@ protected:
 	 *
 	 * @return true if successful, false otherwise.
 	 */
-	bool JsonObjToManifest( TSharedRef< FJsonObject > InJsonObj, FString InNamespace, TSharedRef< FInternationalizationManifest > Manifest );
+	static bool JsonObjToManifest( TSharedRef< FJsonObject > InJsonObj, FString InNamespace, TSharedRef< FInternationalizationManifest > Manifest );
 
 	/**
 	 * Takes a Internationalization manifest and arranges the data into a hierarchy based on namespace.
@@ -101,14 +119,14 @@ protected:
 	 * @param InManifest The Internationalization manifest.
 	 * @param RootElement The root element of the structured data.
 	 */
-	void GenerateStructuredData( TSharedRef< const FInternationalizationManifest > InManifest, TSharedPtr< FStructuredEntry > RootElement );
+	static void GenerateStructuredData( TSharedRef< const FInternationalizationManifest > InManifest, TSharedPtr< FStructuredEntry > RootElement );
 
 	/**
 	 * Goes through the structured, hierarchy based, manifest data and does a non-culture specific sort on namespaces, default text, and key.
 	 *
 	 * @param RootElement The root element of the structured data.
 	 */
-	void SortStructuredData( TSharedPtr< FStructuredEntry > InElement );
+	static void SortStructuredData( TSharedPtr< FStructuredEntry > InElement );
 
 	/**
 	 * Populates a JSON object from Internationalization manifest data that has been structured based on namespace.
@@ -116,7 +134,7 @@ protected:
 	 * @param InElement Internationalization manifest data structured based on namespace.
 	 * @param JsonObj JSON object to be populated.
 	 */
-	void StructuredDataToJsonObj( TSharedPtr< const FStructuredEntry> InElement, TSharedRef< FJsonObject > JsonObj );
+	static void StructuredDataToJsonObj( TSharedPtr< const FStructuredEntry> InElement, TSharedRef< FJsonObject > JsonObj );
 
 public:
 
